Add SlateTests low-level unit test project

* Create new "SlateTests" program for housing Slate/SlateCore tests.
* Add dummy test that passes.
* Add Slate.xml so Horde will run thess tests as part of the "Low Level Tests" run.

diff --git a/Engine/Build/LowLevelTests/Slate.xml b/Engine/Build/LowLevelTests/Slate.xml
new file mode 100644
--- /dev/null
+++ b/Engine/Build/LowLevelTests/Slate.xml
@@ -0,0 +1,10 @@
+﻿<?xml version="1.0" encoding="utf-8"?>
+<BuildGraph xmlns:xsi="http://www.w3.org/2001/XMLSchema-instance" xsi:schemaLocation="http://www.epicgames.com/BuildGraph ../../Build/Graph/Schema.xsd" xmlns="http://www.epicgames.com/BuildGraph">
+  <Option Name="RunSlateTests" DefaultValue="" Description="Run Slate Tests" />
+  <Property Name="TestNames" Value="$(TestNames);Slate" />
+  <Extend Name="RunAllTests">
+    <Expand Name="DeployAndTest" Platform="Win64" TestName="Slate" ShortName="Slate" TargetName="SlateTests" BinaryRelativePath="Engine\Binaries" ReportType="xml" ExtraCompilationArgs="-allmodules" />
+    <Expand Name="DeployAndTest" Platform="Linux" TestName="Slate" ShortName="Slate" TargetName="SlateTests" BinaryRelativePath="Engine\Binaries" ReportType="xml" ExtraCompilationArgs="-allmodules" />
+    <Expand Name="DeployAndTest" Platform="Mac" TestName="Slate" ShortName="Slate" TargetName="SlateTests" BinaryRelativePath="Engine\Binaries" ReportType="xml" ExtraCompilationArgs="-allmodules" />
+  </Extend>
+</BuildGraph>
diff --git a/Engine/Source/Programs/SlateTests/SlateTests.Build.cs b/Engine/Source/Programs/SlateTests/SlateTests.Build.cs
new file mode 100644
--- /dev/null
+++ b/Engine/Source/Programs/SlateTests/SlateTests.Build.cs
@@ -0,0 +1,51 @@
+// Copyright Epic Games, Inc. All Rights Reserved.
+using UnrealBuildTool;
+
+public class SlateTests : TestModuleRules
+{
+	static SlateTests()
+	{
+		if (InTestMode)
+		{
+			TestMetadata = new Metadata();
+			TestMetadata.TestName = "Slate";
+			TestMetadata.TestShortName = "Slate";
+			TestMetadata.ReportType = "xml";
+			TestMetadata.SupportedPlatforms.Add(UnrealTargetPlatform.Linux);
+			TestMetadata.SupportedPlatforms.Add(UnrealTargetPlatform.Mac);
+
+			string PlatformCompilationArgs;
+			foreach (var Platform in UnrealTargetPlatform.GetValidPlatforms())
+			{
+				if (Platform == UnrealTargetPlatform.Android)
+				{
+					PlatformCompilationArgs = "-allmodules -architectures=arm64";
+				}
+				else
+				{
+					PlatformCompilationArgs = "-allmodules";
+				}
+				TestMetadata.PlatformCompilationExtraArgs.Add(Platform, PlatformCompilationArgs);
+			}
+		}
+	}
+
+	public SlateTests(ReadOnlyTargetRules Target) : base(Target, true)
+	{
+		PrivateDependencyModuleNames.AddRange(
+			new string[] {
+				"Core",
+				"CoreUObject",
+				"Slate",
+				"SlateCore",
+			});
+
+		if (Target.bBuildWithEditorOnlyData)
+		{
+			PrivateDependencyModuleNames.AddRange(
+				new string[] {
+					"DesktopPlatform"
+				});
+		}
+	}
+}
diff --git a/Engine/Source/Programs/SlateTests/SlateTests.Target.cs b/Engine/Source/Programs/SlateTests/SlateTests.Target.cs
new file mode 100644
--- /dev/null
+++ b/Engine/Source/Programs/SlateTests/SlateTests.Target.cs
@@ -0,0 +1,16 @@
+// Copyright Epic Games, Inc. All Rights Reserved.
+
+using UnrealBuildTool;
+
+[SupportedPlatforms(UnrealPlatformClass.All)]
+public class SlateTestsTarget : TestTargetRules
+{
+	public SlateTestsTarget(TargetInfo Target) : base(Target)
+	{
+		bTestsRequireApplicationCore = true;
+		bTestsRequireCoreUObject = true;
+
+		bBuildWithEditorOnlyData = Target.Platform.IsInGroup(UnrealPlatformGroup.Desktop)
+			&& (Target.Configuration == UnrealTargetConfiguration.Debug || Target.Configuration == UnrealTargetConfiguration.Development);
+	}
+}
diff --git a/Engine/Source/Programs/SlateTests/Tests/Slate/SlateTest.cpp b/Engine/Source/Programs/SlateTests/Tests/Slate/SlateTest.cpp
new file mode 100644
--- /dev/null
+++ b/Engine/Source/Programs/SlateTests/Tests/Slate/SlateTest.cpp
@@ -0,0 +1,8 @@
+// Copyright Epic Games, Inc. All Rights Reserved.
+
+#include "TestHarness.h"
+
+TEST_CASE("Runtime::Slate::dummy", "[Slate]")
+{
+	CHECK(true);
+}
diff --git a/Engine/Source/Programs/SlateTests/Tests/TestGroupEvents.cpp b/Engine/Source/Programs/SlateTests/Tests/TestGroupEvents.cpp
new file mode 100644
--- /dev/null
+++ b/Engine/Source/Programs/SlateTests/Tests/TestGroupEvents.cpp
@@ -0,0 +1,28 @@
+// Copyright Epic Games, Inc. All Rights Reserved.
+
+#if !EXPLICIT_TESTS_TARGET && WITH_LOW_LEVEL_TESTS
+
+#include "TestCommon/CoreUtilities.h"
+#include "TestCommon/CoreUObjectUtilities.h"
+
+#include <catch2/catch_test_macros.hpp>
+
+GROUP_BEFORE_GLOBAL(Catch::DefaultGroup)
+{
+	InitTaskGraph();
+	InitThreadPool(true);
+#if WITH_COREUOBJECT
+	InitCoreUObject();
+#endif
+}
+
+GROUP_AFTER_GLOBAL(Catch::DefaultGroup)
+{
+#if WITH_COREUOBJECT
+	CleanupCoreUObject();
+#endif
+	CleanupThreadPool();
+	CleanupTaskGraph();
+}
+
+#endif // !EXPLICIT_TESTS_TARGET && WITH_LOW_LEVEL_TESTS
