Add ToolMenus support for priority-based resizing of toolbars

diff --git a/Engine/Source/Developer/ToolMenus/Private/ToolMenus.cpp b/Engine/Source/Developer/ToolMenus/Private/ToolMenus.cpp
--- a/Engine/Source/Developer/ToolMenus/Private/ToolMenus.cpp
+++ b/Engine/Source/Developer/ToolMenus/Private/ToolMenus.cpp
@@ -1469,7 +1469,6 @@ void UToolMenus::PopulateToolBarBuilderWithTopLevelChildren(
 		const TTuple<UToolMenu*, FToolMenuEntry*, FName> SubMenuBlock = SubMenuBlocks.Pop(EAllowShrinking::No);
 		UToolMenu* const SubMenu = SubMenuBlock.Get<0>();
 		FToolMenuEntry* const Block = SubMenuBlock.Get<1>();
-		const FName SectionName = SubMenuBlock.Get<2>();
 
 		// Keep track of how many blocks we've visited to ensure we don't loop indefinitely.
 		if (++NumIterations > 5000)
@@ -1602,14 +1601,18 @@ void UToolMenus::PopulateToolBarBuilderWithTopLevelChildren(
 			);
 
 			const FName UnsetExtensionHook = NAME_None;
-			ToolBarBuilder.AddSeparator(UnsetExtensionHook, VisibilityOverride);
+
+			FMenuEntryResizeParams ResizeParams;
+			ResizeParams.VisibleInOverflow = false;
+
+			ToolBarBuilder.AddSeparator(UnsetExtensionHook, VisibilityOverride, ResizeParams);
 			bHasRaisedEntrySinceLastSeparator = false;
 		}
 
 		// Make sure we actually add the entry if the reason we added a separator above was that the section names changed.
 		if (Entry->Type != EMultiBlockType::Separator)
 		{
-			const bool bRaiseToTopLevel = true;
+			constexpr bool bRaiseToTopLevel = true;
 			PopulateToolBarBuilderWithEntry(ToolBarBuilder, SubMenu, *Entry, bRaiseToTopLevel);
 			bHasRaisedEntrySinceLastSeparator = true;
 
@@ -1647,6 +1650,10 @@ void UToolMenus::PopulateToolBarBuilderWithTopLevelChildren(
 			}
 		);
 
+		FMenuEntryResizeParams ResizeParams;
+		// Never show spacers in overflow menus.
+		ResizeParams.VisibleInOverflow = false;
+
 		const FToolBarStyle ToolbarStyle = ParentMenu->StyleSet->GetWidgetStyle<FToolBarStyle>(ParentMenu->StyleName);
 		ToolBarBuilder.AddWidget(
 			SNew(SSpacer).Size(FVector2D(ToolbarStyle.RaisedChildrenRightPadding, 0)),
@@ -1654,7 +1661,8 @@ void UToolMenus::PopulateToolBarBuilderWithTopLevelChildren(
 			TutorialHighlightName,
 			bSearchable,
 			CustomMenuDelegate,
-			SpacerVisibilityOverride
+			SpacerVisibilityOverride,
+			ResizeParams
 		);
 	}
 }
@@ -1734,6 +1742,12 @@ void UToolMenus::PopulateToolBarBuilderWithEntry(
 		ToolbarLabelOverride = FText();
 	}
 
+	FMenuEntryResizeParams ActualResizeParams = Block.ToolBarData.ResizeParams;
+	if (bIsRaisingToTopLevel && !Block.ToolBarData.ResizeParams.VisibleInOverflow.IsSet())
+	{
+		ActualResizeParams.VisibleInOverflow = false;
+	}
+
 	if (Block.Type == EMultiBlockType::ToolBarButton || (Block.Type == EMultiBlockType::MenuEntry && !Block.IsSubMenu()))
 	{
 		if (Block.Command.IsValid() && !Block.IsCommandKeybindOnly())
@@ -1761,7 +1775,8 @@ void UToolMenus::PopulateToolBarBuilderWithEntry(
 				Block.TutorialHighlightName,
 				FNewMenuDelegate(),
 				VisibilityOverride,
-				ToolbarLabelOverride
+				ToolbarLabelOverride,
+				ActualResizeParams
 			);
 
 			if (bPopCommandList)
@@ -1782,7 +1797,8 @@ void UToolMenus::PopulateToolBarBuilderWithEntry(
 				Block.UserInterfaceActionType,
 				Block.TutorialHighlightName,
 				VisibilityOverride,
-				ToolbarLabelOverride
+				ToolbarLabelOverride,
+				ActualResizeParams
 			);
 		}
 		else
@@ -1796,7 +1812,8 @@ void UToolMenus::PopulateToolBarBuilderWithEntry(
 				Block.UserInterfaceActionType,
 				Block.TutorialHighlightName,
 				VisibilityOverride,
-				ToolbarLabelOverride
+				ToolbarLabelOverride,
+				ActualResizeParams
 			);
 		}
 
@@ -1815,7 +1832,8 @@ void UToolMenus::PopulateToolBarBuilderWithEntry(
 				VisibilityOverride,
 				ToolbarLabelOverride,
 				Block.ToolBarData.PlacementOverride,
-				Block.UserInterfaceActionType
+				Block.UserInterfaceActionType,
+				ActualResizeParams
 			);
 		}
 	}
@@ -1837,7 +1855,8 @@ void UToolMenus::PopulateToolBarBuilderWithEntry(
 				VisibilityOverride,
 				ToolbarLabelOverride,
 				Block.ToolBarData.PlacementOverride,
-				Block.UserInterfaceActionType
+				Block.UserInterfaceActionType,
+				ActualResizeParams
 			);
 		}
 		else
@@ -1863,7 +1882,8 @@ void UToolMenus::PopulateToolBarBuilderWithEntry(
 				VisibilityOverride,
 				ToolbarLabelOverride,
 				Block.ToolBarData.PlacementOverride,
-				Block.UserInterfaceActionType
+				Block.UserInterfaceActionType,
+				ActualResizeParams
 			);
 
 			// Also add any top-level flagged children to the toolbar.
@@ -1967,7 +1987,18 @@ void UToolMenus::PopulateToolBarBuilder(FToolBarBuilder& ToolBarBuilder, UToolMe
 			StyleParams.SizeRule = FSizeParam::ESizeRule::SizeRule_StretchContent;
 			StyleParams.MinSize = 0.5f;
 
-			ToolBarBuilder.AddWidget(SNew(SSpacer), StyleParams, NAME_None, true, {});
+			FMenuEntryResizeParams ResizeParams;
+			// Never show spacers in overflow menus.
+			ResizeParams.VisibleInOverflow = false;
+			// Prevent our spacers from overflowing. This will keep them in the toolbar, but they will still shrink so
+			// overflow of other entries will not be affected. What will be affected, however, is stretching of the
+			// toolbar after entries have been clipped. By keeping the stretchers in the toolbar, we can grow them to
+			// fill up any space left over when a widget is clipped.
+			ResizeParams.AllowClipping = false;
+
+			ToolBarBuilder.AddWidget(
+				SNew(SSpacer), StyleParams, NAME_None, false, FNewMenuDelegate(), TAttribute<EVisibility>(), ResizeParams
+			);
 		}
 
 		// Keep track if this is the first section we're adding for the CurrentAlignment. Make an exception if the
diff --git a/Engine/Source/Developer/ToolMenus/Public/ToolMenuEntry.h b/Engine/Source/Developer/ToolMenus/Public/ToolMenuEntry.h
--- a/Engine/Source/Developer/ToolMenus/Public/ToolMenuEntry.h
+++ b/Engine/Source/Developer/ToolMenus/Public/ToolMenuEntry.h
@@ -65,6 +65,8 @@ public:
 
 	/** Optional override placement to customize menu placement once opened via e.g. a toolbar menu button */
 	TAttribute<EMenuPlacement> PlacementOverride;
+
+	FMenuEntryResizeParams ResizeParams;
 };
 
 
diff --git a/Engine/Source/Programs/SlateTests/Tests/Slate/SlateTest.cpp b/Engine/Source/Programs/SlateTests/Tests/Slate/SlateTest.cpp
--- a/Engine/Source/Programs/SlateTests/Tests/Slate/SlateTest.cpp
+++ b/Engine/Source/Programs/SlateTests/Tests/Slate/SlateTest.cpp
@@ -113,9 +113,8 @@ TEST_CASE("Runtime::Slate::PrioritizedResizing with constrained size and two wid
 	const FMargin WrapButtonPadding = FMargin(4.0f, 0);
 	constexpr int32 WrapButtonIndex = -1;
 	TArray<UE::Slate::FClippingInfo> InOutClippingInfos = {
-		UE::Slate::FClippingInfo{ .ResizeParams = FMenuEntryResizeParams{ .AllowClipping = true }, .X = 0.0,  .Width = 50.0 },
-		UE::Slate::FClippingInfo{
-								  .ResizeParams = FMenuEntryResizeParams{ .AllowClipping = false }, .X = 50.0, .Width = 50.0 },
+		UE::Slate::FClippingInfo{ .ResizeParams = FMenuEntryResizeParams{ .AllowClipping = true },  .X = 0.0,  .Width = 50.0 },
+		UE::Slate::FClippingInfo{ .ResizeParams = FMenuEntryResizeParams{ .AllowClipping = false }, .X = 50.0, .Width = 50.0 },
 	};
 	double WrapButtonX = 0.0;
 
@@ -130,7 +129,7 @@ TEST_CASE("Runtime::Slate::PrioritizedResizing with constrained size and two wid
 TEST_CASE("Runtime::Slate::PrioritizedResizing superfluous space after clipping is distributed to stretching widgets", "[Slate]")
 {
 	constexpr double WrapButtonWidth = 0.0;
-	const FMargin WrapButtonPadding = FMargin(0.0f, 0);
+	const FMargin WrapButtonPadding = FMargin(0);
 	constexpr int32 WrapButtonIndex = -1;
 	const TArray<UE::Slate::FClippingInfo> OriginalClippingInfos = {
 		UE::Slate::FClippingInfo{ .X = 0.0, .Width = 50.0 },
@@ -198,7 +197,7 @@ TEST_CASE("Runtime::Slate::PrioritizedResizing superfluous space after clipping
 TEST_CASE("Runtime::Slate::PrioritizedResizing all superfluous space is eaten by stretching widgets", "[Slate]")
 {
 	constexpr double WrapButtonWidth = 0.0;
-	const FMargin WrapButtonPadding = FMargin(0.0f, 0);
+	const FMargin WrapButtonPadding = FMargin(0);
 	constexpr int32 WrapButtonIndex = -1;
 	const TArray<UE::Slate::FClippingInfo> OriginalClippingInfos = {
 		UE::Slate::FClippingInfo{ .X = 0.0, .Width = 50.0 },
@@ -233,7 +232,7 @@ TEST_CASE("Runtime::Slate::PrioritizedResizing all superfluous space is eaten by
 TEST_CASE("Runtime::Slate::PrioritizedResizing sorts zero-width widgets identical to input regardless of clipping priority", "[Slate]")
 {
 	constexpr double WrapButtonWidth = 0.0;
-	const FMargin WrapButtonPadding = FMargin(0.0f, 0);
+	const FMargin WrapButtonPadding = FMargin(0);
 	constexpr int32 WrapButtonIndex = -1;
 
 	const TArray<UE::Slate::FClippingInfo> OriginalClippingInfos = {
@@ -244,7 +243,7 @@ TEST_CASE("Runtime::Slate::PrioritizedResizing sorts zero-width widgets identica
 	};
 	double WrapButtonX = 0.0;
 
-	SECTION("When size is contrained")
+	SECTION("When size is constrained")
 	{
 		constexpr double AllottedWidth = 50.0;
 
@@ -266,7 +265,7 @@ TEST_CASE("Runtime::Slate::PrioritizedResizing sorts zero-width widgets identica
 TEST_CASE("Runtime::Slate::PrioritizedResizing positions a wrap button at index 1 snugly", "[Slate]")
 {
 	constexpr double WrapButtonWidth = 40;
-	const FMargin WrapButtonPadding = FMargin(0.0f, 0);
+	const FMargin WrapButtonPadding = FMargin(0);
 	constexpr int32 WrapButtonIndex = 1;
 
 	const TArray<UE::Slate::FClippingInfo> OriginalClippingInfos = {
@@ -277,7 +276,7 @@ TEST_CASE("Runtime::Slate::PrioritizedResizing positions a wrap button at index
 	};
 	double WrapButtonX = 0.0;
 
-	SECTION("When size is contrained")
+	SECTION("When size is constrained")
 	{
 		constexpr double AllottedWidth = 190;
 
@@ -301,3 +300,296 @@ TEST_CASE("Runtime::Slate::PrioritizedResizing positions a wrap button at index
 		}
 	}
 }
+
+TEST_CASE("Runtime::Slate::PrioritizedResizing adds wrap button when needed", "[Slate]")
+{
+	constexpr double WrapButtonWidth = 40;
+	const FMargin WrapButtonPadding = FMargin(0);
+	constexpr int32 WrapButtonIndex = 1;
+
+	const TArray<UE::Slate::FClippingInfo> OriginalClippingInfos = {
+		UE::Slate::FClippingInfo{ .X = 0, .Width = 50 },
+		UE::Slate::FClippingInfo{ .X = 50, .Width = 50, .ResizeParams = { .VisibleInOverflow = false } },
+		UE::Slate::FClippingInfo{ .X = 100, .Width = 50, .ResizeParams = { .VisibleInOverflow = false } },
+		UE::Slate::FClippingInfo{ .X = 150, .Width = 50 },
+	};
+	double WrapButtonX = 0.0;
+
+	SECTION("When size is constrained")
+	{
+		constexpr double AllottedWidth = 190;
+
+		TArray<UE::Slate::FClippingInfo> InOutClippingInfos = OriginalClippingInfos;
+		UE::Slate::PrioritizedResize(
+			AllottedWidth, WrapButtonWidth, WrapButtonPadding, WrapButtonIndex, InOutClippingInfos, WrapButtonX
+		);
+
+		SECTION("The non-right-most widgets are not clipped")
+		{
+			CHECK(InOutClippingInfos[0].bWasClipped == false);
+			CHECK(InOutClippingInfos[1].bWasClipped == false);
+			CHECK(InOutClippingInfos[2].bWasClipped == false);
+		}
+
+		SECTION("The right-most widget is clipped")
+		{
+			CHECK(InOutClippingInfos[3].bWasClipped == true);
+		}
+
+		SECTION("The wrap button is positioned")
+		{
+			CHECK(WrapButtonX == 50);
+		}
+	}
+}
+
+TEST_CASE("Runtime::Slate::PrioritizedResizing with complex set of entries", "[Slate]")
+{
+	constexpr double WrapButtonWidth = 40;
+	const FMargin WrapButtonPadding = FMargin(0);
+	constexpr int32 WrapButtonIndex = 1;
+
+	// clang-format off
+	const TArray<UE::Slate::FClippingInfo> OriginalClippingInfos = {
+		UE::Slate::FClippingInfo{ .X =   0, .Width = 50, .ResizeParams = { .ClippingPriority = 1000 } },
+		UE::Slate::FClippingInfo{ .X =  50, .Width = 50, .ResizeParams = { .ClippingPriority =  500, .VisibleInOverflow = false } },
+		UE::Slate::FClippingInfo{ .X = 100, .Width = 50, .ResizeParams = { .ClippingPriority =  500, .VisibleInOverflow = false } },
+		UE::Slate::FClippingInfo{ .X = 150, .Width = 50, .ResizeParams = { .ClippingPriority =  800 } },
+		UE::Slate::FClippingInfo{ .X = 200, .Width = 50, .ResizeParams = { .ClippingPriority =  400, .VisibleInOverflow = false } },
+		UE::Slate::FClippingInfo{ .X = 250, .Width = 50, .ResizeParams = { .ClippingPriority =  400, .VisibleInOverflow = false } },
+		UE::Slate::FClippingInfo{ .X = 300, .Width = 50, .bIsStretchable = true, .ResizeParams = { .AllowClipping = false } },
+		UE::Slate::FClippingInfo{ .X = 350, .Width = 50, .ResizeParams = { .ClippingPriority = 1000, .VisibleInOverflow = false } },
+		UE::Slate::FClippingInfo{ .X = 400, .Width = 50, .ResizeParams = { .AllowClipping = false } },
+		UE::Slate::FClippingInfo{ .X = 450, .Width = 50, .ResizeParams = { .VisibleInOverflow = false, .AllowClipping = false } },
+	};
+	// clang-format on
+	double WrapButtonX = -1.0;
+
+	GIVEN("Size is not constrained")
+	{
+		constexpr double AllottedWidth = 501;
+
+		TArray<UE::Slate::FClippingInfo> InOutClippingInfos = OriginalClippingInfos;
+		UE::Slate::PrioritizedResize(
+			AllottedWidth, WrapButtonWidth, WrapButtonPadding, WrapButtonIndex, InOutClippingInfos, WrapButtonX
+		);
+
+		THEN("No widgets are clipped")
+		{
+			for (const UE::Slate::FClippingInfo& Info : InOutClippingInfos)
+			{
+				CHECK(Info.bWasClipped == false);
+			}
+		}
+
+		THEN("The wrap button is not placed")
+		{
+			CHECK(WrapButtonX < 0.0);
+		}
+	}
+
+	GIVEN("Size is constrained by one pixel")
+	{
+		constexpr double AllottedWidth = 499;
+
+		TArray<UE::Slate::FClippingInfo> InOutClippingInfos = OriginalClippingInfos;
+		UE::Slate::PrioritizedResize(
+			AllottedWidth, WrapButtonWidth, WrapButtonPadding, WrapButtonIndex, InOutClippingInfos, WrapButtonX
+		);
+
+		THEN("The right-most of the lowest priority entries is clipped")
+		{
+			CHECK(InOutClippingInfos[5].bWasClipped == true);
+		}
+
+		THEN("Only one entry is clipped")
+		{
+			int32 NumClipped = 0;
+			for (const UE::Slate::FClippingInfo& Info : InOutClippingInfos)
+			{
+				if (Info.bWasClipped)
+				{
+					++NumClipped;
+				}
+			}
+			CHECK(NumClipped == 1);
+		}
+
+		THEN("The wrap button is not placed")
+		{
+			CHECK(WrapButtonX < 0.0);
+		}
+	}
+
+	GIVEN("Size is constrained by a little more than one button width")
+	{
+		constexpr double AllottedWidth = 445;
+
+		TArray<UE::Slate::FClippingInfo> InOutClippingInfos = OriginalClippingInfos;
+		UE::Slate::PrioritizedResize(
+			AllottedWidth, WrapButtonWidth, WrapButtonPadding, WrapButtonIndex, InOutClippingInfos, WrapButtonX
+		);
+
+		THEN("Two entries are clipped")
+		{
+			int32 NumClipped = 0;
+			for (const UE::Slate::FClippingInfo& Info : InOutClippingInfos)
+			{
+				if (Info.bWasClipped)
+				{
+					++NumClipped;
+				}
+			}
+			CHECK(NumClipped == 2);
+		}
+
+		THEN("The wrap button is not placed")
+		{
+			CHECK(WrapButtonX < 0.0);
+		}
+	}
+
+	GIVEN("Size is constrained by a little more than two button widths")
+	{
+		constexpr double AllottedWidth = 395;
+
+		TArray<UE::Slate::FClippingInfo> InOutClippingInfos = OriginalClippingInfos;
+		UE::Slate::PrioritizedResize(
+			AllottedWidth, WrapButtonWidth, WrapButtonPadding, WrapButtonIndex, InOutClippingInfos, WrapButtonX
+		);
+
+		THEN("Three entries are clipped")
+		{
+			int32 NumClipped = 0;
+			for (const UE::Slate::FClippingInfo& Info : InOutClippingInfos)
+			{
+				if (Info.bWasClipped)
+				{
+					++NumClipped;
+				}
+			}
+			CHECK(NumClipped == 3);
+		}
+
+		THEN("The wrap button is not placed")
+		{
+			CHECK(WrapButtonX < 0.0);
+		}
+	}
+
+	GIVEN("Size is constrained by a little more three button widths")
+	{
+		constexpr double AllottedWidth = 345;
+
+		TArray<UE::Slate::FClippingInfo> InOutClippingInfos = OriginalClippingInfos;
+		UE::Slate::PrioritizedResize(
+			AllottedWidth, WrapButtonWidth, WrapButtonPadding, WrapButtonIndex, InOutClippingInfos, WrapButtonX
+		);
+
+		THEN("Four entries are clipped")
+		{
+			int32 NumClipped = 0;
+			for (const UE::Slate::FClippingInfo& Info : InOutClippingInfos)
+			{
+				if (Info.bWasClipped)
+				{
+					++NumClipped;
+				}
+			}
+			CHECK(NumClipped == 4);
+		}
+
+		THEN("The wrap button is not placed")
+		{
+			CHECK(WrapButtonX < 0.0);
+		}
+	}
+}
+
+TEST_CASE("Runtime::Slate::PrioritizedResizing takes the wrapping button width into account", "[Slate]")
+{
+	constexpr double WrapButtonWidth = 40.0;
+	const FMargin WrapButtonPadding = FMargin(3.0f, 0);
+	constexpr int32 WrapButtonIndex = -2;
+
+	const TArray<UE::Slate::FClippingInfo> OriginalClippingInfos = {
+		UE::Slate::FClippingInfo{ .X = 0,   .Width = 70 },
+		UE::Slate::FClippingInfo{ .X = 70,  .Width = 70 },
+		UE::Slate::FClippingInfo{ .X = 140, .Width = 70 },
+		UE::Slate::FClippingInfo{ .X = 210, .Width = 70 },
+	};
+	double WrapButtonX = -1.0;
+
+	GIVEN("Size is constrained by a few pixels")
+	{
+		constexpr double AllottedWidth = 274.0;
+
+		TArray<UE::Slate::FClippingInfo> InOutClippingInfos = OriginalClippingInfos;
+		UE::Slate::PrioritizedResize(
+			AllottedWidth, WrapButtonWidth, WrapButtonPadding, WrapButtonIndex, InOutClippingInfos, WrapButtonX
+		);
+
+		THEN("The wrapping button is placed")
+		{
+			// Position + Left padding
+			CHECK(WrapButtonX == 140 + 3);
+		}
+
+		THEN("The last widget is clipped")
+		{
+			CHECK(InOutClippingInfos[3].bWasClipped);
+		}
+
+		THEN("Only one widget is clipped")
+		{
+			int32 NumClipped = 0;
+			for (const UE::Slate::FClippingInfo& Info : InOutClippingInfos)
+			{
+				if (Info.bWasClipped)
+				{
+					++NumClipped;
+				}
+			}
+			CHECK(NumClipped == 1);
+		}
+	}
+
+	GIVEN("Size is constrained enough so the wrap button will force another widget to clip")
+	{
+		constexpr double AllottedWidth = 230.0;
+
+		TArray<UE::Slate::FClippingInfo> InOutClippingInfos = OriginalClippingInfos;
+		UE::Slate::PrioritizedResize(
+			AllottedWidth, WrapButtonWidth, WrapButtonPadding, WrapButtonIndex, InOutClippingInfos, WrapButtonX
+		);
+
+		THEN("The wrapping button is placed")
+		{
+			// Position + Left padding
+			CHECK(WrapButtonX == 70 + 3);
+		}
+
+		THEN("The second last widget is clipped")
+		{
+			CHECK(InOutClippingInfos[2].bWasClipped);
+		}
+
+		THEN("The last widget is clipped")
+		{
+			CHECK(InOutClippingInfos[3].bWasClipped);
+		}
+
+		THEN("Two widgets are clipped")
+		{
+			int32 NumClipped = 0;
+			for (const UE::Slate::FClippingInfo& Info : InOutClippingInfos)
+			{
+				if (Info.bWasClipped)
+				{
+					++NumClipped;
+				}
+			}
+			CHECK(NumClipped == 2);
+		}
+	}
+}
diff --git a/Engine/Source/Runtime/Slate/Private/Framework/MultiBox/SClippingHorizontalBox.cpp b/Engine/Source/Runtime/Slate/Private/Framework/MultiBox/SClippingHorizontalBox.cpp
--- a/Engine/Source/Runtime/Slate/Private/Framework/MultiBox/SClippingHorizontalBox.cpp
+++ b/Engine/Source/Runtime/Slate/Private/Framework/MultiBox/SClippingHorizontalBox.cpp
@@ -22,20 +22,11 @@ void PrioritizedResize(
 )
 {
 	double WidthOfAllChildren = 0.0;
-	double WidthOfAllChildrenWhenWrapButtonIsNeeded = 0.0;
 	double NonClippingWidgetWidths = 0.0;
 	for (const UE::Slate::FClippingInfo& Info : InOutClippingInfos)
 	{
 		WidthOfAllChildren += Info.Width;
 
-		// Find the combined width of widgets that either appear in the overflow menu or never clips, i.e. the widgets
-		// that compete with the wrap button for space. Any other widgets than these we'll just clip before adding a wrap button.
-		if (Info.ResizeParams.VisibleInOverflow.Get(FMenuEntryResizeParams::DefaultVisibleInOverflow)
-			|| !Info.ResizeParams.AllowClipping.Get(FMenuEntryResizeParams::DefaultAllowClipping))
-		{
-			WidthOfAllChildrenWhenWrapButtonIsNeeded += Info.Width;
-		}
-
 		// Find the combined width of widgets that should never clip.
 		if (!Info.ResizeParams.AllowClipping.Get(FMenuEntryResizeParams::DefaultAllowClipping))
 		{
@@ -44,7 +35,6 @@ void PrioritizedResize(
 	}
 
 	const bool bNeedsClipping = WidthOfAllChildren > InAllottedWidth;
-	const bool bNeedsWrapButton = WidthOfAllChildrenWhenWrapButtonIsNeeded > InAllottedWidth;
 
 	// Early out if we don't need to clip.
 	if (!bNeedsClipping)
@@ -52,6 +42,15 @@ void PrioritizedResize(
 		return;
 	}
 
+	// Propagate the resize parameter's visibility overflow to the clipping info.
+	for (UE::Slate::FClippingInfo& Info : InOutClippingInfos)
+	{
+		Info.bAppearsInOverflow = Info.ResizeParams.VisibleInOverflow.Get(FMenuEntryResizeParams::DefaultVisibleInOverflow);
+	}
+
+	const double PaddedWidthOfWrapButton = InWrapButtonWidth + InWrapButtonPadding.Left + InWrapButtonPadding.Right;
+
+	bool bNeedsWrapButton = false;
 	// Clip widgets in priority order.
 	{
 		// Construct a priority-sorted array of pointers into the clipping infos. Work on this copy
@@ -73,30 +72,44 @@ void PrioritizedResize(
 		// once their combined width is larger than our allotted width. Start counting at the width required by
 		// non-clipping widgets to make sure we start clipping clippable widgets immediately once non-clippable widths
 		// and the accumulator sum up to our allotted width.
+		//
+		// If we find that we need a wrap button, restart the clipping loop and include the width of the warp button in
+		// the initial accumulator width.
+		for (int32 NumClippingAttemptsLeft = 1; NumClippingAttemptsLeft > 0; --NumClippingAttemptsLeft)
 		{
 			double WidthAccumulator = NonClippingWidgetWidths;
 			// If we need a wrap button, clip widgets as if the wrap button is already taking up space.
 			if (bNeedsWrapButton)
 			{
-				WidthAccumulator += InWrapButtonWidth + InWrapButtonPadding.Left + InWrapButtonPadding.Right;
+				WidthAccumulator += PaddedWidthOfWrapButton;
 			}
 
-			for (UE::Slate::FClippingInfo* Info : PrioritySortedClippingInfos)
+			for (int32 InfoIndex = 0; InfoIndex < PrioritySortedClippingInfos.Num(); ++InfoIndex)
 			{
+				UE::Slate::FClippingInfo& Info = *PrioritySortedClippingInfos[InfoIndex];
+
 				// Don't clip non-clipping widgets and also don't count their width, we've already done that above.
-				if (!Info->ResizeParams.AllowClipping.Get(FMenuEntryResizeParams::DefaultAllowClipping))
+				if (!Info.ResizeParams.AllowClipping.Get(FMenuEntryResizeParams::DefaultAllowClipping))
 				{
 					continue;
 				}
 
-				WidthAccumulator += Info->Width;
+				WidthAccumulator += Info.Width;
 
 				// We are wider than our allotted width, so mark this widget for clipping.
 				if (WidthAccumulator > InAllottedWidth)
 				{
-					Info->bWasClipped = true;
-					Info->bAppearsInOverflow =
-						Info->ResizeParams.VisibleInOverflow.Get(FMenuEntryResizeParams::DefaultVisibleInOverflow);
+					Info.bWasClipped = true;
+
+					// If this widget appears in an overflow menu, but we haven't taken the wrap button width
+					// into account, we need to restart the clipping loop and do it over again while taking the wrap button width into account.
+					if (Info.bAppearsInOverflow && !bNeedsWrapButton)
+					{
+						bNeedsWrapButton = true;
+
+						++NumClippingAttemptsLeft;
+						break;
+					}
 				}
 			}
 		}
@@ -122,8 +135,8 @@ void PrioritizedResize(
 		Info.X -= WidthOfRemovedWidgets;
 	}
 
-	// Here we want to EXPAND again to fill up the whole post-clip toolbar if possible. The reason is because we've
-	// reduced the width of the toolbar in "integer steps", i.e. by clipping whole widgets, and in general we'll have
+	// Here we want to EXPAND again to fill up the whole post-clip toolbar if possible. The reason is we've reduced
+	// the width of the toolbar in "integer steps", i.e. by clipping whole widgets, and in general we'll have
 	// some slack space left.
 	//
 	// An example of widgets to expand are variable-sized spacers that were previously shrunk to make things fit.
@@ -141,11 +154,11 @@ void PrioritizedResize(
 		if (NumStretchWidgets > 0)
 		{
 			// Distribute our extra space to our stretching children.
-			const double ExtraSpace =
-				InAllottedWidth - WidthOfAllChildren
-				- (bNeedsWrapButton ? InWrapButtonWidth + InWrapButtonPadding.Left + InWrapButtonPadding.Right : 0.0)
-				+ WidthOfRemovedWidgets;
+			const double ExtraSpace = InAllottedWidth - WidthOfAllChildren
+									- (bNeedsWrapButton ? PaddedWidthOfWrapButton : 0.0) + WidthOfRemovedWidgets;
 			const double ExtraSpacePerStretchChild = ExtraSpace / NumStretchWidgets;
+
+			if (ExtraSpace > 0.0)
 			{
 				double StretchSpaceAdded = 0.0;
 				for (UE::Slate::FClippingInfo& Info : InOutClippingInfos)
@@ -219,7 +232,6 @@ void PrioritizedResize(
 			OutWrapButtonX = InAllottedWidth - InWrapButtonWidth - InWrapButtonPadding.Right;
 		}
 
-		const double PaddedWidthOfWrapButton = InWrapButtonWidth + InWrapButtonPadding.Left + InWrapButtonPadding.Right;
 		for (int32 Index = IndexFromWhichToPushToTheRight; !NonClippedInfos.IsEmpty() && Index < NonClippedInfos.Num();
 			 ++Index)
 		{
@@ -298,7 +310,7 @@ void SClippingHorizontalBox::OnArrangeChildren(const FGeometry& AllottedGeometry
 
 	for (const UE::Slate::FClippingInfo& Info : ClippingInfos)
 	{
-		if (Info.bAppearsInOverflow)
+		if (Info.bWasClipped && Info.bAppearsInOverflow)
 		{
 			ClippedWidgets.Add(Info.Widget.ToWeakPtr());
 		}
