Fix Win64 and Linux warnings in SlateTests

diff --git a/Engine/Source/Programs/SlateTests/Tests/Slate/SlateTest.cpp b/Engine/Source/Programs/SlateTests/Tests/Slate/SlateTest.cpp
--- a/Engine/Source/Programs/SlateTests/Tests/Slate/SlateTest.cpp
+++ b/Engine/Source/Programs/SlateTests/Tests/Slate/SlateTest.cpp
@@ -203,7 +203,7 @@ TEST_CASE("Runtime::Slate::PrioritizedResizing all superfluous space is eaten by
 		UE::Slate::FClippingInfo{ .X = 0.0, .Width = 50.0 },
 		UE::Slate::FClippingInfo{ .X = 50.0, .Width = 50.0, .bIsStretchable = true },
 		UE::Slate::FClippingInfo{ .X = 100.0, .Width = 50.0 },
-		UE::Slate::FClippingInfo{ .X = 150.0, .Width = 50.0, .ResizeParams{ .ClippingPriority = -100 } }, // Clips  first.
+		UE::Slate::FClippingInfo{ .ResizeParams{ .ClippingPriority = -100 }, .X = 150.0, .Width = 50.0 }, // Clips  first.
 		UE::Slate::FClippingInfo{ .X = 200.0, .Width = 50.0 },
 	};
 	double WrapButtonX = 0.0;
@@ -236,10 +236,10 @@ TEST_CASE("Runtime::Slate::PrioritizedResizing sorts zero-width widgets identica
 	constexpr int32 WrapButtonIndex = -1;
 
 	const TArray<UE::Slate::FClippingInfo> OriginalClippingInfos = {
-		UE::Slate::FClippingInfo{ .X = 0, .Width = 0,   .ResizeParams{ .ClippingPriority = 2 }, .bIsStretchable = true },
-		UE::Slate::FClippingInfo{ .X = 0, .Width = 0,   .ResizeParams{ .ClippingPriority = 4 }, .bIsStretchable = true },
-		UE::Slate::FClippingInfo{ .X = 0, .Width = 100, .ResizeParams{ .ClippingPriority = 3 }, .bIsStretchable = true },
-		UE::Slate::FClippingInfo{ .X = 0, .Width = 0,   .ResizeParams{ .ClippingPriority = 1 }, .bIsStretchable = true },
+		UE::Slate::FClippingInfo{ .ResizeParams{ .ClippingPriority = 2 }, .X = 0, .Width = 0,   .bIsStretchable = true },
+		UE::Slate::FClippingInfo{ .ResizeParams{ .ClippingPriority = 4 }, .X = 0, .Width = 0,   .bIsStretchable = true },
+		UE::Slate::FClippingInfo{ .ResizeParams{ .ClippingPriority = 3 }, .X = 0, .Width = 100, .bIsStretchable = true },
+		UE::Slate::FClippingInfo{ .ResizeParams{ .ClippingPriority = 1 }, .X = 0, .Width = 0,   .bIsStretchable = true },
 	};
 	double WrapButtonX = 0.0;
 
@@ -309,8 +309,8 @@ TEST_CASE("Runtime::Slate::PrioritizedResizing adds wrap button when needed", "[
 
 	const TArray<UE::Slate::FClippingInfo> OriginalClippingInfos = {
 		UE::Slate::FClippingInfo{ .X = 0, .Width = 50 },
-		UE::Slate::FClippingInfo{ .X = 50, .Width = 50, .ResizeParams = { .VisibleInOverflow = false } },
-		UE::Slate::FClippingInfo{ .X = 100, .Width = 50, .ResizeParams = { .VisibleInOverflow = false } },
+		UE::Slate::FClippingInfo{ .ResizeParams = { .VisibleInOverflow = false }, .X = 50, .Width = 50 },
+		UE::Slate::FClippingInfo{ .ResizeParams = { .VisibleInOverflow = false }, .X = 100, .Width = 50 },
 		UE::Slate::FClippingInfo{ .X = 150, .Width = 50 },
 	};
 	double WrapButtonX = 0.0;
@@ -351,16 +351,16 @@ TEST_CASE("Runtime::Slate::PrioritizedResizing with complex set of entries", "[S
 
 	// clang-format off
 	const TArray<UE::Slate::FClippingInfo> OriginalClippingInfos = {
-		UE::Slate::FClippingInfo{ .X =   0, .Width = 50, .ResizeParams = { .ClippingPriority = 1000 } },
-		UE::Slate::FClippingInfo{ .X =  50, .Width = 50, .ResizeParams = { .ClippingPriority =  500, .VisibleInOverflow = false } },
-		UE::Slate::FClippingInfo{ .X = 100, .Width = 50, .ResizeParams = { .ClippingPriority =  500, .VisibleInOverflow = false } },
-		UE::Slate::FClippingInfo{ .X = 150, .Width = 50, .ResizeParams = { .ClippingPriority =  800 } },
-		UE::Slate::FClippingInfo{ .X = 200, .Width = 50, .ResizeParams = { .ClippingPriority =  400, .VisibleInOverflow = false } },
-		UE::Slate::FClippingInfo{ .X = 250, .Width = 50, .ResizeParams = { .ClippingPriority =  400, .VisibleInOverflow = false } },
-		UE::Slate::FClippingInfo{ .X = 300, .Width = 50, .bIsStretchable = true, .ResizeParams = { .AllowClipping = false } },
-		UE::Slate::FClippingInfo{ .X = 350, .Width = 50, .ResizeParams = { .ClippingPriority = 1000, .VisibleInOverflow = false } },
-		UE::Slate::FClippingInfo{ .X = 400, .Width = 50, .ResizeParams = { .AllowClipping = false } },
-		UE::Slate::FClippingInfo{ .X = 450, .Width = 50, .ResizeParams = { .VisibleInOverflow = false, .AllowClipping = false } },
+		UE::Slate::FClippingInfo{ .ResizeParams = { .ClippingPriority = 1000 }, .X = 0, .Width = 50 },
+		UE::Slate::FClippingInfo{ .ResizeParams = { .ClippingPriority = 500, .VisibleInOverflow = false }, .X = 50, .Width = 50 },
+		UE::Slate::FClippingInfo{ .ResizeParams = { .ClippingPriority = 500, .VisibleInOverflow = false }, .X = 100, .Width = 50 },
+		UE::Slate::FClippingInfo{ .ResizeParams = { .ClippingPriority = 800 }, .X = 150, .Width = 50 },
+		UE::Slate::FClippingInfo{ .ResizeParams = { .ClippingPriority = 400, .VisibleInOverflow = false }, .X = 200, .Width = 50 },
+		UE::Slate::FClippingInfo{ .ResizeParams = { .ClippingPriority = 400, .VisibleInOverflow = false }, .X = 250, .Width = 50 },
+		UE::Slate::FClippingInfo{ .ResizeParams = { .AllowClipping = false }, .X = 300, .Width = 50, .bIsStretchable = true },
+		UE::Slate::FClippingInfo{ .ResizeParams = { .ClippingPriority = 1000, .VisibleInOverflow = false }, .X = 350, .Width = 50 },
+		UE::Slate::FClippingInfo{ .ResizeParams = { .AllowClipping = false }, .X = 400, .Width = 50 },
+		UE::Slate::FClippingInfo{ .ResizeParams = { .AllowClipping = false, .VisibleInOverflow = false  }, .X = 450, .Width = 50 },
 	};
 	// clang-format on
 	double WrapButtonX = -1.0;
