Rename the SlateTests project's SlateTest.cpp to SlatePrioritizedResizingTest.cpp to reflect its contents

diff --git a/Engine/Source/Programs/SlateTests/Tests/Slate/SlatePrioritizedResizingTest.cpp b/Engine/Source/Programs/SlateTests/Tests/Slate/SlatePrioritizedResizingTest.cpp
new file mode 100644
--- /dev/null
+++ b/Engine/Source/Programs/SlateTests/Tests/Slate/SlatePrioritizedResizingTest.cpp
@@ -0,0 +1,595 @@
+// Copyright Epic Games, Inc. All Rights Reserved.
+
+#include "TestHarness.h"
+
+#include "Framework/MultiBox/SClippingHorizontalBox.h"
+
+#include <iostream>
+
+namespace UE::Slate
+{
+
+bool operator==(const UE::Slate::FClippingInfo& Left, const UE::Slate::FClippingInfo& Right)
+{
+	return Left.Widget == Right.Widget && Left.X == Right.X && Left.Width == Right.Width
+		&& Left.bIsStretchable == Right.bIsStretchable && Left.bAppearsInOverflow == Right.bAppearsInOverflow
+		&& Left.bWasClipped == Right.bWasClipped;
+}
+
+std::ostream& operator<<(std::ostream& Stream, const UE::Slate::FClippingInfo& Value)
+{
+	// Format booleans as "true" and "false" instead of 1 and 0.
+	Stream << std::boolalpha;
+
+	Stream << "FClippingInfo3("
+		   << "Widget=" << Value.Widget.Get() << ", ";
+	Stream << "FMenuEntryResizeParams(" << Value.ResizeParams.ClippingPriority.Get() << ", "
+		   << Value.ResizeParams.AllowClipping.Get() << ", " << Value.ResizeParams.VisibleInOverflow.Get() << "), ";
+	Stream << "X=" << Value.X << ", ";
+	Stream << "Width=" << Value.Width << ", ";
+	Stream << "bIsStretchable=" << Value.bIsStretchable << ", ";
+	Stream << "bAppearsInOverflow=" << Value.bAppearsInOverflow << ", ";
+	Stream << "bWasClipped=" << Value.bWasClipped;
+	Stream << ")";
+	return Stream;
+}
+
+} // namespace UE::Slate
+
+TEST_CASE("Runtime::Slate::PrioritizedResizing works on empty input", "[Slate]")
+{
+	constexpr double AllottedWidth = 200.0;
+	constexpr double WrapButtonWidth = 40.0;
+	const FMargin WrapButtonPadding = FMargin(4.0f, 0);
+	constexpr int32 WrapButtonIndex = -1;
+	TArray<UE::Slate::FClippingInfo> ClippingInfos;
+
+	double WrapButtonX = 0.0;
+	UE::Slate::PrioritizedResize(
+		AllottedWidth, WrapButtonWidth, WrapButtonPadding, WrapButtonIndex, ClippingInfos, WrapButtonX
+	);
+
+	CHECK(ClippingInfos.IsEmpty());
+}
+
+TEST_CASE("Runtime::Slate::PrioritizedResizing clips no widgets when there's plenty of space", "[Slate]")
+{
+	constexpr double AllottedWidth = 500.0;
+	constexpr double WrapButtonWidth = 40.0;
+	const FMargin WrapButtonPadding = FMargin(4.0f, 0);
+	constexpr int32 WrapButtonIndex = -1;
+	TArray<UE::Slate::FClippingInfo> ClippingInfos = {
+		UE::Slate::FClippingInfo{ .X = 0.0,  .Width = 20.0  },
+		UE::Slate::FClippingInfo{ .X = 30.0, .Width = 100.0 },
+	};
+	double WrapButtonX = 0.0;
+
+	UE::Slate::PrioritizedResize(
+		AllottedWidth, WrapButtonWidth, WrapButtonPadding, WrapButtonIndex, ClippingInfos, WrapButtonX
+	);
+
+	for (const UE::Slate::FClippingInfo& Info : ClippingInfos)
+	{
+		CHECK_FALSE(Info.bWasClipped);
+	}
+}
+
+TEST_CASE("Runtime::Slate::PrioritizedResizing with plenty of space does not modify input with default resize params", "[Slate]")
+{
+	constexpr double AllottedWidth = 1000.0;
+	constexpr double WrapButtonWidth = 40.0;
+	const FMargin WrapButtonPadding = FMargin(4.0f, 0);
+	constexpr int32 WrapButtonIndex = -1;
+	const TArray<UE::Slate::FClippingInfo> OriginalClippingInfos = {
+		UE::Slate::FClippingInfo{ .X = 0.0,   .Width = 20.0  },
+        UE::Slate::FClippingInfo{ .X = 20.0,  .Width = 100.0 },
+		UE::Slate::FClippingInfo{ .X = 120.0, .Width = 30.0  },
+        UE::Slate::FClippingInfo{ .X = 150.0, .Width = 40.0  },
+		UE::Slate::FClippingInfo{ .X = 190.0, .Width = 200.0 },
+        UE::Slate::FClippingInfo{ .X = 390.0, .Width = 15.0  },
+	};
+	double WrapButtonX = 0.0;
+
+	TArray<UE::Slate::FClippingInfo> InOutClippingInfos = OriginalClippingInfos;
+	UE::Slate::PrioritizedResize(
+		AllottedWidth, WrapButtonWidth, WrapButtonPadding, WrapButtonIndex, InOutClippingInfos, WrapButtonX
+	);
+
+	CHECK(OriginalClippingInfos.Num() == InOutClippingInfos.Num());
+
+	for (int32 Index = 0; Index < OriginalClippingInfos.Num(); ++Index)
+	{
+		const UE::Slate::FClippingInfo Original = OriginalClippingInfos[Index];
+		const UE::Slate::FClippingInfo Output = InOutClippingInfos[Index];
+
+		CHECK(Original == Output);
+	}
+}
+
+TEST_CASE("Runtime::Slate::PrioritizedResizing with constrained size and two widgets clips the clippable one", "[Slate]")
+{
+	constexpr double AllottedWidth = 90.0;
+	constexpr double WrapButtonWidth = 40.0;
+	const FMargin WrapButtonPadding = FMargin(4.0f, 0);
+	constexpr int32 WrapButtonIndex = -1;
+	TArray<UE::Slate::FClippingInfo> InOutClippingInfos = {
+		UE::Slate::FClippingInfo{ .ResizeParams = FMenuEntryResizeParams{ .AllowClipping = true },  .X = 0.0,  .Width = 50.0 },
+		UE::Slate::FClippingInfo{ .ResizeParams = FMenuEntryResizeParams{ .AllowClipping = false }, .X = 50.0, .Width = 50.0 },
+	};
+	double WrapButtonX = 0.0;
+
+	UE::Slate::PrioritizedResize(
+		AllottedWidth, WrapButtonWidth, WrapButtonPadding, WrapButtonIndex, InOutClippingInfos, WrapButtonX
+	);
+
+	CHECK(InOutClippingInfos[0].bWasClipped == true);
+	CHECK(InOutClippingInfos[1].bWasClipped == false);
+}
+
+TEST_CASE("Runtime::Slate::PrioritizedResizing superfluous space after clipping is distributed to stretching widgets", "[Slate]")
+{
+	constexpr double WrapButtonWidth = 0.0;
+	const FMargin WrapButtonPadding = FMargin(0);
+	constexpr int32 WrapButtonIndex = -1;
+	const TArray<UE::Slate::FClippingInfo> OriginalClippingInfos = {
+		UE::Slate::FClippingInfo{ .X = 0.0, .Width = 50.0 },
+		UE::Slate::FClippingInfo{ .X = 50.0, .Width = 50.0, .bIsStretchable = true },
+		UE::Slate::FClippingInfo{ .X = 100.0, .Width = 50.0 },
+		UE::Slate::FClippingInfo{ .X = 150.0, .Width = 50.0, .bIsStretchable = true },
+		UE::Slate::FClippingInfo{ .X = 200.0, .Width = 50.0 },
+		UE::Slate::FClippingInfo{ .X = 250.0, .Width = 50.0 },
+	};
+	double WrapButtonX = 0.0;
+
+	// The tests following this rely on the total width of all widgets to be 300 pixels.
+	{
+		double TotalWidthOfWidgets = 0.0;
+		for (const UE::Slate::FClippingInfo& Info : OriginalClippingInfos)
+		{
+			TotalWidthOfWidgets += Info.Width;
+		}
+
+		CHECK(TotalWidthOfWidgets == 300.0);
+	}
+
+	SECTION("When all widgets don't fit")
+	{
+		constexpr double AllottedWidth = 295.0;
+
+		TArray<UE::Slate::FClippingInfo> InOutClippingInfos = OriginalClippingInfos;
+		UE::Slate::PrioritizedResize(
+			AllottedWidth, WrapButtonWidth, WrapButtonPadding, WrapButtonIndex, InOutClippingInfos, WrapButtonX
+		);
+
+		SECTION("Only the last one is clipped")
+		{
+			for (int Index = 0; Index < InOutClippingInfos.Num() - 1; ++Index)
+			{
+				INFO("Index=" << Index);
+				CHECK(InOutClippingInfos[Index].bWasClipped == false);
+			}
+
+			CHECK(InOutClippingInfos.Last().bWasClipped);
+		}
+
+		SECTION("The non-stretching widgets keep their original width")
+		{
+			CHECK(InOutClippingInfos[0].Width == 50.0);
+			// Stretching widget
+			CHECK(InOutClippingInfos[2].Width == 50.0);
+			// Stretching widget
+			CHECK(InOutClippingInfos[4].Width == 50.0);
+			CHECK(InOutClippingInfos[5].Width == 50.0);
+		}
+
+		SECTION("The stretching widgets are expanded")
+		{
+			// Non-stretching widget
+			CHECK(InOutClippingInfos[1].Width > 51.0);
+			// Non-stretching widget
+			CHECK(InOutClippingInfos[3].Width > 51.0);
+			// Non-stretching widget
+			// Non-stretching widget
+		}
+	}
+}
+
+TEST_CASE("Runtime::Slate::PrioritizedResizing all superfluous space is eaten by stretching widgets", "[Slate]")
+{
+	constexpr double WrapButtonWidth = 0.0;
+	const FMargin WrapButtonPadding = FMargin(0);
+	constexpr int32 WrapButtonIndex = -1;
+	const TArray<UE::Slate::FClippingInfo> OriginalClippingInfos = {
+		UE::Slate::FClippingInfo{ .X = 0.0, .Width = 50.0 },
+		UE::Slate::FClippingInfo{ .X = 50.0, .Width = 50.0, .bIsStretchable = true },
+		UE::Slate::FClippingInfo{ .X = 100.0, .Width = 50.0 },
+		UE::Slate::FClippingInfo{ .ResizeParams{ .ClippingPriority = -100 }, .X = 150.0, .Width = 50.0 }, // Clips  first.
+		UE::Slate::FClippingInfo{ .X = 200.0, .Width = 50.0 },
+	};
+	double WrapButtonX = 0.0;
+
+	SECTION("When all widgets don't fit")
+	{
+		constexpr double AllottedWidth = 245.0;
+
+		TArray<UE::Slate::FClippingInfo> InOutClippingInfos = OriginalClippingInfos;
+		UE::Slate::PrioritizedResize(
+			AllottedWidth, WrapButtonWidth, WrapButtonPadding, WrapButtonIndex, InOutClippingInfos, WrapButtonX
+		);
+
+		SECTION("The stretching widget eats all space")
+		{
+			CHECK(InOutClippingInfos[1].Width == AllottedWidth - 50 - 50 - 50);
+		}
+
+		SECTION("The widget to the right of the stretching widget is moved to the right")
+		{
+			CHECK(InOutClippingInfos[2].X == AllottedWidth - 50 - 50);
+		}
+	}
+}
+
+TEST_CASE("Runtime::Slate::PrioritizedResizing sorts zero-width widgets identical to input regardless of clipping priority", "[Slate]")
+{
+	constexpr double WrapButtonWidth = 0.0;
+	const FMargin WrapButtonPadding = FMargin(0);
+	constexpr int32 WrapButtonIndex = -1;
+
+	const TArray<UE::Slate::FClippingInfo> OriginalClippingInfos = {
+		UE::Slate::FClippingInfo{ .ResizeParams{ .ClippingPriority = 2 }, .X = 0, .Width = 0,   .bIsStretchable = true },
+		UE::Slate::FClippingInfo{ .ResizeParams{ .ClippingPriority = 4 }, .X = 0, .Width = 0,   .bIsStretchable = true },
+		UE::Slate::FClippingInfo{ .ResizeParams{ .ClippingPriority = 3 }, .X = 0, .Width = 100, .bIsStretchable = true },
+		UE::Slate::FClippingInfo{ .ResizeParams{ .ClippingPriority = 1 }, .X = 0, .Width = 0,   .bIsStretchable = true },
+	};
+	double WrapButtonX = 0.0;
+
+	SECTION("When size is constrained")
+	{
+		constexpr double AllottedWidth = 50.0;
+
+		TArray<UE::Slate::FClippingInfo> InOutClippingInfos = OriginalClippingInfos;
+		UE::Slate::PrioritizedResize(
+			AllottedWidth, WrapButtonWidth, WrapButtonPadding, WrapButtonIndex, InOutClippingInfos, WrapButtonX
+		);
+
+		SECTION("The widget sorting hasn't changed")
+		{
+			CHECK(InOutClippingInfos[0].ResizeParams.ClippingPriority.Get() == 2);
+			CHECK(InOutClippingInfos[1].ResizeParams.ClippingPriority.Get() == 4);
+			CHECK(InOutClippingInfos[2].ResizeParams.ClippingPriority.Get() == 3);
+			CHECK(InOutClippingInfos[3].ResizeParams.ClippingPriority.Get() == 1);
+		}
+	}
+}
+
+TEST_CASE("Runtime::Slate::PrioritizedResizing positions a wrap button at index 1 snugly", "[Slate]")
+{
+	constexpr double WrapButtonWidth = 40;
+	const FMargin WrapButtonPadding = FMargin(0);
+	constexpr int32 WrapButtonIndex = 1;
+
+	const TArray<UE::Slate::FClippingInfo> OriginalClippingInfos = {
+		UE::Slate::FClippingInfo{ .X = 0,   .Width = 50 },
+		UE::Slate::FClippingInfo{ .X = 50,  .Width = 50 },
+		UE::Slate::FClippingInfo{ .X = 100, .Width = 50 },
+		UE::Slate::FClippingInfo{ .X = 150, .Width = 50 },
+	};
+	double WrapButtonX = 0.0;
+
+	SECTION("When size is constrained")
+	{
+		constexpr double AllottedWidth = 190;
+
+		TArray<UE::Slate::FClippingInfo> InOutClippingInfos = OriginalClippingInfos;
+		UE::Slate::PrioritizedResize(
+			AllottedWidth, WrapButtonWidth, WrapButtonPadding, WrapButtonIndex, InOutClippingInfos, WrapButtonX
+		);
+
+		SECTION("The wrap button is positioned snugly with the non-clipped widgets")
+		{
+			CHECK(InOutClippingInfos[0].X == 0);
+			CHECK(InOutClippingInfos[0].Width == 50);
+
+			CHECK(WrapButtonX == 50);
+
+			CHECK(InOutClippingInfos[1].X == 90);
+			CHECK(InOutClippingInfos[1].Width == 50);
+
+			CHECK(InOutClippingInfos[2].X == 140);
+			CHECK(InOutClippingInfos[2].Width == 50);
+		}
+	}
+}
+
+TEST_CASE("Runtime::Slate::PrioritizedResizing adds wrap button when needed", "[Slate]")
+{
+	constexpr double WrapButtonWidth = 40;
+	const FMargin WrapButtonPadding = FMargin(0);
+	constexpr int32 WrapButtonIndex = 1;
+
+	const TArray<UE::Slate::FClippingInfo> OriginalClippingInfos = {
+		UE::Slate::FClippingInfo{ .X = 0, .Width = 50 },
+		UE::Slate::FClippingInfo{ .ResizeParams = { .VisibleInOverflow = false }, .X = 50, .Width = 50 },
+		UE::Slate::FClippingInfo{ .ResizeParams = { .VisibleInOverflow = false }, .X = 100, .Width = 50 },
+		UE::Slate::FClippingInfo{ .X = 150, .Width = 50 },
+	};
+	double WrapButtonX = 0.0;
+
+	SECTION("When size is constrained")
+	{
+		constexpr double AllottedWidth = 190;
+
+		TArray<UE::Slate::FClippingInfo> InOutClippingInfos = OriginalClippingInfos;
+		UE::Slate::PrioritizedResize(
+			AllottedWidth, WrapButtonWidth, WrapButtonPadding, WrapButtonIndex, InOutClippingInfos, WrapButtonX
+		);
+
+		SECTION("The non-right-most widgets are not clipped")
+		{
+			CHECK(InOutClippingInfos[0].bWasClipped == false);
+			CHECK(InOutClippingInfos[1].bWasClipped == false);
+			CHECK(InOutClippingInfos[2].bWasClipped == false);
+		}
+
+		SECTION("The right-most widget is clipped")
+		{
+			CHECK(InOutClippingInfos[3].bWasClipped == true);
+		}
+
+		SECTION("The wrap button is positioned")
+		{
+			CHECK(WrapButtonX == 50);
+		}
+	}
+}
+
+TEST_CASE("Runtime::Slate::PrioritizedResizing with complex set of entries", "[Slate]")
+{
+	constexpr double WrapButtonWidth = 40;
+	const FMargin WrapButtonPadding = FMargin(0);
+	constexpr int32 WrapButtonIndex = 1;
+
+	// clang-format off
+	const TArray<UE::Slate::FClippingInfo> OriginalClippingInfos = {
+		UE::Slate::FClippingInfo{ .ResizeParams = { .ClippingPriority = 1000 }, .X = 0, .Width = 50 },
+		UE::Slate::FClippingInfo{ .ResizeParams = { .ClippingPriority = 500, .VisibleInOverflow = false }, .X = 50, .Width = 50 },
+		UE::Slate::FClippingInfo{ .ResizeParams = { .ClippingPriority = 500, .VisibleInOverflow = false }, .X = 100, .Width = 50 },
+		UE::Slate::FClippingInfo{ .ResizeParams = { .ClippingPriority = 800 }, .X = 150, .Width = 50 },
+		UE::Slate::FClippingInfo{ .ResizeParams = { .ClippingPriority = 400, .VisibleInOverflow = false }, .X = 200, .Width = 50 },
+		UE::Slate::FClippingInfo{ .ResizeParams = { .ClippingPriority = 400, .VisibleInOverflow = false }, .X = 250, .Width = 50 },
+		UE::Slate::FClippingInfo{ .ResizeParams = { .AllowClipping = false }, .X = 300, .Width = 50, .bIsStretchable = true },
+		UE::Slate::FClippingInfo{ .ResizeParams = { .ClippingPriority = 1000, .VisibleInOverflow = false }, .X = 350, .Width = 50 },
+		UE::Slate::FClippingInfo{ .ResizeParams = { .AllowClipping = false }, .X = 400, .Width = 50 },
+		UE::Slate::FClippingInfo{ .ResizeParams = { .AllowClipping = false, .VisibleInOverflow = false  }, .X = 450, .Width = 50 },
+	};
+	// clang-format on
+	double WrapButtonX = -1.0;
+
+	GIVEN("Size is not constrained")
+	{
+		constexpr double AllottedWidth = 501;
+
+		TArray<UE::Slate::FClippingInfo> InOutClippingInfos = OriginalClippingInfos;
+		UE::Slate::PrioritizedResize(
+			AllottedWidth, WrapButtonWidth, WrapButtonPadding, WrapButtonIndex, InOutClippingInfos, WrapButtonX
+		);
+
+		THEN("No widgets are clipped")
+		{
+			for (const UE::Slate::FClippingInfo& Info : InOutClippingInfos)
+			{
+				CHECK(Info.bWasClipped == false);
+			}
+		}
+
+		THEN("The wrap button is not placed")
+		{
+			CHECK(WrapButtonX < 0.0);
+		}
+	}
+
+	GIVEN("Size is constrained by one pixel")
+	{
+		constexpr double AllottedWidth = 499;
+
+		TArray<UE::Slate::FClippingInfo> InOutClippingInfos = OriginalClippingInfos;
+		UE::Slate::PrioritizedResize(
+			AllottedWidth, WrapButtonWidth, WrapButtonPadding, WrapButtonIndex, InOutClippingInfos, WrapButtonX
+		);
+
+		THEN("The right-most of the lowest priority entries is clipped")
+		{
+			CHECK(InOutClippingInfos[5].bWasClipped == true);
+		}
+
+		THEN("Only one entry is clipped")
+		{
+			int32 NumClipped = 0;
+			for (const UE::Slate::FClippingInfo& Info : InOutClippingInfos)
+			{
+				if (Info.bWasClipped)
+				{
+					++NumClipped;
+				}
+			}
+			CHECK(NumClipped == 1);
+		}
+
+		THEN("The wrap button is not placed")
+		{
+			CHECK(WrapButtonX < 0.0);
+		}
+	}
+
+	GIVEN("Size is constrained by a little more than one button width")
+	{
+		constexpr double AllottedWidth = 445;
+
+		TArray<UE::Slate::FClippingInfo> InOutClippingInfos = OriginalClippingInfos;
+		UE::Slate::PrioritizedResize(
+			AllottedWidth, WrapButtonWidth, WrapButtonPadding, WrapButtonIndex, InOutClippingInfos, WrapButtonX
+		);
+
+		THEN("Two entries are clipped")
+		{
+			int32 NumClipped = 0;
+			for (const UE::Slate::FClippingInfo& Info : InOutClippingInfos)
+			{
+				if (Info.bWasClipped)
+				{
+					++NumClipped;
+				}
+			}
+			CHECK(NumClipped == 2);
+		}
+
+		THEN("The wrap button is not placed")
+		{
+			CHECK(WrapButtonX < 0.0);
+		}
+	}
+
+	GIVEN("Size is constrained by a little more than two button widths")
+	{
+		constexpr double AllottedWidth = 395;
+
+		TArray<UE::Slate::FClippingInfo> InOutClippingInfos = OriginalClippingInfos;
+		UE::Slate::PrioritizedResize(
+			AllottedWidth, WrapButtonWidth, WrapButtonPadding, WrapButtonIndex, InOutClippingInfos, WrapButtonX
+		);
+
+		THEN("Three entries are clipped")
+		{
+			int32 NumClipped = 0;
+			for (const UE::Slate::FClippingInfo& Info : InOutClippingInfos)
+			{
+				if (Info.bWasClipped)
+				{
+					++NumClipped;
+				}
+			}
+			CHECK(NumClipped == 3);
+		}
+
+		THEN("The wrap button is not placed")
+		{
+			CHECK(WrapButtonX < 0.0);
+		}
+	}
+
+	GIVEN("Size is constrained by a little more three button widths")
+	{
+		constexpr double AllottedWidth = 345;
+
+		TArray<UE::Slate::FClippingInfo> InOutClippingInfos = OriginalClippingInfos;
+		UE::Slate::PrioritizedResize(
+			AllottedWidth, WrapButtonWidth, WrapButtonPadding, WrapButtonIndex, InOutClippingInfos, WrapButtonX
+		);
+
+		THEN("Four entries are clipped")
+		{
+			int32 NumClipped = 0;
+			for (const UE::Slate::FClippingInfo& Info : InOutClippingInfos)
+			{
+				if (Info.bWasClipped)
+				{
+					++NumClipped;
+				}
+			}
+			CHECK(NumClipped == 4);
+		}
+
+		THEN("The wrap button is not placed")
+		{
+			CHECK(WrapButtonX < 0.0);
+		}
+	}
+}
+
+TEST_CASE("Runtime::Slate::PrioritizedResizing takes the wrapping button width into account", "[Slate]")
+{
+	constexpr double WrapButtonWidth = 40.0;
+	const FMargin WrapButtonPadding = FMargin(3.0f, 0);
+	constexpr int32 WrapButtonIndex = -2;
+
+	const TArray<UE::Slate::FClippingInfo> OriginalClippingInfos = {
+		UE::Slate::FClippingInfo{ .X = 0,   .Width = 70 },
+		UE::Slate::FClippingInfo{ .X = 70,  .Width = 70 },
+		UE::Slate::FClippingInfo{ .X = 140, .Width = 70 },
+		UE::Slate::FClippingInfo{ .X = 210, .Width = 70 },
+	};
+	double WrapButtonX = -1.0;
+
+	GIVEN("Size is constrained by a few pixels")
+	{
+		constexpr double AllottedWidth = 274.0;
+
+		TArray<UE::Slate::FClippingInfo> InOutClippingInfos = OriginalClippingInfos;
+		UE::Slate::PrioritizedResize(
+			AllottedWidth, WrapButtonWidth, WrapButtonPadding, WrapButtonIndex, InOutClippingInfos, WrapButtonX
+		);
+
+		THEN("The wrapping button is placed")
+		{
+			// Position + Left padding
+			CHECK(WrapButtonX == 140 + 3);
+		}
+
+		THEN("The last widget is clipped")
+		{
+			CHECK(InOutClippingInfos[3].bWasClipped);
+		}
+
+		THEN("Only one widget is clipped")
+		{
+			int32 NumClipped = 0;
+			for (const UE::Slate::FClippingInfo& Info : InOutClippingInfos)
+			{
+				if (Info.bWasClipped)
+				{
+					++NumClipped;
+				}
+			}
+			CHECK(NumClipped == 1);
+		}
+	}
+
+	GIVEN("Size is constrained enough so the wrap button will force another widget to clip")
+	{
+		constexpr double AllottedWidth = 230.0;
+
+		TArray<UE::Slate::FClippingInfo> InOutClippingInfos = OriginalClippingInfos;
+		UE::Slate::PrioritizedResize(
+			AllottedWidth, WrapButtonWidth, WrapButtonPadding, WrapButtonIndex, InOutClippingInfos, WrapButtonX
+		);
+
+		THEN("The wrapping button is placed")
+		{
+			// Position + Left padding
+			CHECK(WrapButtonX == 70 + 3);
+		}
+
+		THEN("The second last widget is clipped")
+		{
+			CHECK(InOutClippingInfos[2].bWasClipped);
+		}
+
+		THEN("The last widget is clipped")
+		{
+			CHECK(InOutClippingInfos[3].bWasClipped);
+		}
+
+		THEN("Two widgets are clipped")
+		{
+			int32 NumClipped = 0;
+			for (const UE::Slate::FClippingInfo& Info : InOutClippingInfos)
+			{
+				if (Info.bWasClipped)
+				{
+					++NumClipped;
+				}
+			}
+			CHECK(NumClipped == 2);
+		}
+	}
+}
diff --git a/Engine/Source/Programs/SlateTests/Tests/Slate/SlateTest.cpp b/Engine/Source/Programs/SlateTests/Tests/Slate/SlateTest.cpp
deleted file mode 100644
