Fix issue in SClippingHorizontalBox where the overflow box was always arranged even when it didn't need to be.

diff --git a/Engine/Source/Programs/SlateTests/Tests/Slate/SlatePrioritizedResizingTest.cpp b/Engine/Source/Programs/SlateTests/Tests/Slate/SlatePrioritizedResizingTest.cpp
--- a/Engine/Source/Programs/SlateTests/Tests/Slate/SlatePrioritizedResizingTest.cpp
+++ b/Engine/Source/Programs/SlateTests/Tests/Slate/SlatePrioritizedResizingTest.cpp
@@ -38,13 +38,13 @@ std::ostream& operator<<(std::ostream& Stream, const UE::Slate::FClippingInfo& V
 
 TEST_CASE("Runtime::Slate::PrioritizedResizing works on empty input", "[Slate]")
 {
-	constexpr double AllottedWidth = 200.0;
-	constexpr double WrapButtonWidth = 40.0;
+	constexpr float AllottedWidth = 200.0f;
+	constexpr float WrapButtonWidth = 40.0f;
 	const FMargin WrapButtonPadding = FMargin(4.0f, 0);
 	constexpr int32 WrapButtonIndex = -1;
 	TArray<UE::Slate::FClippingInfo> ClippingInfos;
 
-	double WrapButtonX = 0.0;
+	TOptional<float> WrapButtonX;
 	UE::Slate::PrioritizedResize(
 		AllottedWidth, WrapButtonWidth, WrapButtonPadding, WrapButtonIndex, ClippingInfos, WrapButtonX
 	);
@@ -54,15 +54,15 @@ TEST_CASE("Runtime::Slate::PrioritizedResizing works on empty input", "[Slate]")
 
 TEST_CASE("Runtime::Slate::PrioritizedResizing clips no widgets when there's plenty of space", "[Slate]")
 {
-	constexpr double AllottedWidth = 500.0;
-	constexpr double WrapButtonWidth = 40.0;
+	constexpr float AllottedWidth = 500.0f;
+	constexpr float WrapButtonWidth = 40.0f;
 	const FMargin WrapButtonPadding = FMargin(4.0f, 0);
 	constexpr int32 WrapButtonIndex = -1;
 	TArray<UE::Slate::FClippingInfo> ClippingInfos = {
-		UE::Slate::FClippingInfo{ .X = 0.0,  .Width = 20.0  },
-		UE::Slate::FClippingInfo{ .X = 30.0, .Width = 100.0 },
+		UE::Slate::FClippingInfo{ .X = 0.0f,  .Width = 20.0f  },
+		UE::Slate::FClippingInfo{ .X = 30.0f, .Width = 100.0f },
 	};
-	double WrapButtonX = 0.0;
+	TOptional<float> WrapButtonX;
 
 	UE::Slate::PrioritizedResize(
 		AllottedWidth, WrapButtonWidth, WrapButtonPadding, WrapButtonIndex, ClippingInfos, WrapButtonX
@@ -76,19 +76,19 @@ TEST_CASE("Runtime::Slate::PrioritizedResizing clips no widgets when there's ple
 
 TEST_CASE("Runtime::Slate::PrioritizedResizing with plenty of space does not modify input with default resize params", "[Slate]")
 {
-	constexpr double AllottedWidth = 1000.0;
-	constexpr double WrapButtonWidth = 40.0;
+	constexpr float AllottedWidth = 1000.0f;
+	constexpr float WrapButtonWidth = 40.0f;
 	const FMargin WrapButtonPadding = FMargin(4.0f, 0);
 	constexpr int32 WrapButtonIndex = -1;
 	const TArray<UE::Slate::FClippingInfo> OriginalClippingInfos = {
-		UE::Slate::FClippingInfo{ .X = 0.0,   .Width = 20.0  },
-        UE::Slate::FClippingInfo{ .X = 20.0,  .Width = 100.0 },
-		UE::Slate::FClippingInfo{ .X = 120.0, .Width = 30.0  },
-        UE::Slate::FClippingInfo{ .X = 150.0, .Width = 40.0  },
-		UE::Slate::FClippingInfo{ .X = 190.0, .Width = 200.0 },
-        UE::Slate::FClippingInfo{ .X = 390.0, .Width = 15.0  },
+		UE::Slate::FClippingInfo{ .X = 0.0f, .Width = 20.0f },
+        UE::Slate::FClippingInfo{ .X = 20.0, .Width = 100.0f },
+		UE::Slate::FClippingInfo{ .X = 120.0f, .Width = 30.0f },
+        UE::Slate::FClippingInfo{ .X = 150.0f, .Width = 40.0f },
+		UE::Slate::FClippingInfo{ .X = 190.0f, .Width = 200.0f },
+        UE::Slate::FClippingInfo{ .X = 390.0f, .Width = 15.0f },
 	};
-	double WrapButtonX = 0.0;
+	TOptional<float> WrapButtonX;
 
 	TArray<UE::Slate::FClippingInfo> InOutClippingInfos = OriginalClippingInfos;
 	UE::Slate::PrioritizedResize(
@@ -108,15 +108,15 @@ TEST_CASE("Runtime::Slate::PrioritizedResizing with plenty of space does not mod
 
 TEST_CASE("Runtime::Slate::PrioritizedResizing with constrained size and two widgets clips the clippable one", "[Slate]")
 {
-	constexpr double AllottedWidth = 90.0;
-	constexpr double WrapButtonWidth = 40.0;
+	constexpr float AllottedWidth = 90.0f;
+	constexpr float WrapButtonWidth = 40.0f;
 	const FMargin WrapButtonPadding = FMargin(4.0f, 0);
 	constexpr int32 WrapButtonIndex = -1;
 	TArray<UE::Slate::FClippingInfo> InOutClippingInfos = {
-		UE::Slate::FClippingInfo{ .ResizeParams = FMenuEntryResizeParams{ .AllowClipping = true },  .X = 0.0,  .Width = 50.0 },
-		UE::Slate::FClippingInfo{ .ResizeParams = FMenuEntryResizeParams{ .AllowClipping = false }, .X = 50.0, .Width = 50.0 },
+		UE::Slate::FClippingInfo{ .ResizeParams = FMenuEntryResizeParams{ .AllowClipping = true },  .X = 0.0f,  .Width = 50.0f },
+		UE::Slate::FClippingInfo{ .ResizeParams = FMenuEntryResizeParams{ .AllowClipping = false }, .X = 50.0f, .Width = 50.0f },
 	};
-	double WrapButtonX = 0.0;
+	TOptional<float> WrapButtonX;
 
 	UE::Slate::PrioritizedResize(
 		AllottedWidth, WrapButtonWidth, WrapButtonPadding, WrapButtonIndex, InOutClippingInfos, WrapButtonX
@@ -128,33 +128,33 @@ TEST_CASE("Runtime::Slate::PrioritizedResizing with constrained size and two wid
 
 TEST_CASE("Runtime::Slate::PrioritizedResizing superfluous space after clipping is distributed to stretching widgets", "[Slate]")
 {
-	constexpr double WrapButtonWidth = 0.0;
+	constexpr float WrapButtonWidth = 0.0;
 	const FMargin WrapButtonPadding = FMargin(0);
 	constexpr int32 WrapButtonIndex = -1;
 	const TArray<UE::Slate::FClippingInfo> OriginalClippingInfos = {
-		UE::Slate::FClippingInfo{ .X = 0.0, .Width = 50.0 },
-		UE::Slate::FClippingInfo{ .X = 50.0, .Width = 50.0, .bIsStretchable = true },
-		UE::Slate::FClippingInfo{ .X = 100.0, .Width = 50.0 },
-		UE::Slate::FClippingInfo{ .X = 150.0, .Width = 50.0, .bIsStretchable = true },
-		UE::Slate::FClippingInfo{ .X = 200.0, .Width = 50.0 },
-		UE::Slate::FClippingInfo{ .X = 250.0, .Width = 50.0 },
+		UE::Slate::FClippingInfo{ .X = 0.0f, .Width = 50.0f },
+		UE::Slate::FClippingInfo{ .X = 50.0f, .Width = 50.0f, .bIsStretchable = true },
+		UE::Slate::FClippingInfo{ .X = 100.0f, .Width = 50.0f },
+		UE::Slate::FClippingInfo{ .X = 150.0f, .Width = 50.0f, .bIsStretchable = true },
+		UE::Slate::FClippingInfo{ .X = 200.0f, .Width = 50.0f },
+		UE::Slate::FClippingInfo{ .X = 250.0f, .Width = 50.0f },
 	};
-	double WrapButtonX = 0.0;
+	TOptional<float> WrapButtonX;
 
 	// The tests following this rely on the total width of all widgets to be 300 pixels.
 	{
-		double TotalWidthOfWidgets = 0.0;
+		float TotalWidthOfWidgets = 0.0f;
 		for (const UE::Slate::FClippingInfo& Info : OriginalClippingInfos)
 		{
 			TotalWidthOfWidgets += Info.Width;
 		}
 
-		CHECK(TotalWidthOfWidgets == 300.0);
+		CHECK(TotalWidthOfWidgets == 300.0f);
 	}
 
 	SECTION("When all widgets don't fit")
 	{
-		constexpr double AllottedWidth = 295.0;
+		constexpr float AllottedWidth = 295.0f;
 
 		TArray<UE::Slate::FClippingInfo> InOutClippingInfos = OriginalClippingInfos;
 		UE::Slate::PrioritizedResize(
@@ -174,20 +174,20 @@ TEST_CASE("Runtime::Slate::PrioritizedResizing superfluous space after clipping
 
 		SECTION("The non-stretching widgets keep their original width")
 		{
-			CHECK(InOutClippingInfos[0].Width == 50.0);
+			CHECK(InOutClippingInfos[0].Width == 50.0f);
 			// Stretching widget
-			CHECK(InOutClippingInfos[2].Width == 50.0);
+			CHECK(InOutClippingInfos[2].Width == 50.0f);
 			// Stretching widget
-			CHECK(InOutClippingInfos[4].Width == 50.0);
-			CHECK(InOutClippingInfos[5].Width == 50.0);
+			CHECK(InOutClippingInfos[4].Width == 50.0f);
+			CHECK(InOutClippingInfos[5].Width == 50.0f);
 		}
 
 		SECTION("The stretching widgets are expanded")
 		{
 			// Non-stretching widget
-			CHECK(InOutClippingInfos[1].Width > 51.0);
+			CHECK(InOutClippingInfos[1].Width > 51.0f);
 			// Non-stretching widget
-			CHECK(InOutClippingInfos[3].Width > 51.0);
+			CHECK(InOutClippingInfos[3].Width > 51.0f);
 			// Non-stretching widget
 			// Non-stretching widget
 		}
@@ -196,21 +196,21 @@ TEST_CASE("Runtime::Slate::PrioritizedResizing superfluous space after clipping
 
 TEST_CASE("Runtime::Slate::PrioritizedResizing all superfluous space is eaten by stretching widgets", "[Slate]")
 {
-	constexpr double WrapButtonWidth = 0.0;
-	const FMargin WrapButtonPadding = FMargin(0);
+	constexpr float WrapButtonWidth = 0.0f;
+	const FMargin WrapButtonPadding = FMargin(0.0f);
 	constexpr int32 WrapButtonIndex = -1;
 	const TArray<UE::Slate::FClippingInfo> OriginalClippingInfos = {
-		UE::Slate::FClippingInfo{ .X = 0.0, .Width = 50.0 },
-		UE::Slate::FClippingInfo{ .X = 50.0, .Width = 50.0, .bIsStretchable = true },
-		UE::Slate::FClippingInfo{ .X = 100.0, .Width = 50.0 },
-		UE::Slate::FClippingInfo{ .ResizeParams{ .ClippingPriority = -100 }, .X = 150.0, .Width = 50.0 }, // Clips  first.
-		UE::Slate::FClippingInfo{ .X = 200.0, .Width = 50.0 },
+		UE::Slate::FClippingInfo{ .X = 0.0f, .Width = 50.0f },
+		UE::Slate::FClippingInfo{ .X = 50.0f, .Width = 50.0f, .bIsStretchable = true },
+		UE::Slate::FClippingInfo{ .X = 100.0f, .Width = 50.0f },
+		UE::Slate::FClippingInfo{ .ResizeParams{ .ClippingPriority = -100.0f }, .X = 150.0f, .Width = 50.0f }, // Clips  first.
+		UE::Slate::FClippingInfo{ .X = 200.0f, .Width = 50.0f },
 	};
-	double WrapButtonX = 0.0;
+	TOptional<float> WrapButtonX;
 
 	SECTION("When all widgets don't fit")
 	{
-		constexpr double AllottedWidth = 245.0;
+		constexpr float AllottedWidth = 245.0f;
 
 		TArray<UE::Slate::FClippingInfo> InOutClippingInfos = OriginalClippingInfos;
 		UE::Slate::PrioritizedResize(
@@ -219,33 +219,33 @@ TEST_CASE("Runtime::Slate::PrioritizedResizing all superfluous space is eaten by
 
 		SECTION("The stretching widget eats all space")
 		{
-			CHECK(InOutClippingInfos[1].Width == AllottedWidth - 50 - 50 - 50);
+			CHECK(InOutClippingInfos[1].Width == AllottedWidth - 50.0f - 50.0f - 50.0f);
 		}
 
 		SECTION("The widget to the right of the stretching widget is moved to the right")
 		{
-			CHECK(InOutClippingInfos[2].X == AllottedWidth - 50 - 50);
+			CHECK(InOutClippingInfos[2].X == AllottedWidth - 50.0f - 50.0f);
 		}
 	}
 }
 
 TEST_CASE("Runtime::Slate::PrioritizedResizing sorts zero-width widgets identical to input regardless of clipping priority", "[Slate]")
 {
-	constexpr double WrapButtonWidth = 0.0;
-	const FMargin WrapButtonPadding = FMargin(0);
+	constexpr float WrapButtonWidth = 0.0f;
+	const FMargin WrapButtonPadding = FMargin(0.0f);
 	constexpr int32 WrapButtonIndex = -1;
 
 	const TArray<UE::Slate::FClippingInfo> OriginalClippingInfos = {
-		UE::Slate::FClippingInfo{ .ResizeParams{ .ClippingPriority = 2 }, .X = 0, .Width = 0,   .bIsStretchable = true },
-		UE::Slate::FClippingInfo{ .ResizeParams{ .ClippingPriority = 4 }, .X = 0, .Width = 0,   .bIsStretchable = true },
-		UE::Slate::FClippingInfo{ .ResizeParams{ .ClippingPriority = 3 }, .X = 0, .Width = 100, .bIsStretchable = true },
-		UE::Slate::FClippingInfo{ .ResizeParams{ .ClippingPriority = 1 }, .X = 0, .Width = 0,   .bIsStretchable = true },
+		UE::Slate::FClippingInfo{ .ResizeParams{ .ClippingPriority = 2 }, .X = 0.0f, .Width = 0.0f,   .bIsStretchable = true },
+		UE::Slate::FClippingInfo{ .ResizeParams{ .ClippingPriority = 4 }, .X = 0.0f, .Width = 0.0f,   .bIsStretchable = true },
+		UE::Slate::FClippingInfo{ .ResizeParams{ .ClippingPriority = 3 }, .X = 0.0f, .Width = 100.0f, .bIsStretchable = true },
+		UE::Slate::FClippingInfo{ .ResizeParams{ .ClippingPriority = 1 }, .X = 0.0f, .Width = 0.0f,   .bIsStretchable = true },
 	};
-	double WrapButtonX = 0.0;
+	TOptional<float> WrapButtonX;
 
 	SECTION("When size is constrained")
 	{
-		constexpr double AllottedWidth = 50.0;
+		constexpr float AllottedWidth = 50.0f;
 
 		TArray<UE::Slate::FClippingInfo> InOutClippingInfos = OriginalClippingInfos;
 		UE::Slate::PrioritizedResize(
@@ -264,21 +264,21 @@ TEST_CASE("Runtime::Slate::PrioritizedResizing sorts zero-width widgets identica
 
 TEST_CASE("Runtime::Slate::PrioritizedResizing positions a wrap button at index 1 snugly", "[Slate]")
 {
-	constexpr double WrapButtonWidth = 40;
+	constexpr float WrapButtonWidth = 40;
 	const FMargin WrapButtonPadding = FMargin(0);
 	constexpr int32 WrapButtonIndex = 1;
 
 	const TArray<UE::Slate::FClippingInfo> OriginalClippingInfos = {
-		UE::Slate::FClippingInfo{ .X = 0,   .Width = 50 },
-		UE::Slate::FClippingInfo{ .X = 50,  .Width = 50 },
-		UE::Slate::FClippingInfo{ .X = 100, .Width = 50 },
-		UE::Slate::FClippingInfo{ .X = 150, .Width = 50 },
+		UE::Slate::FClippingInfo{ .X = 0.0f,   .Width = 50.0f },
+		UE::Slate::FClippingInfo{ .X = 50.0f,  .Width = 50.0f },
+		UE::Slate::FClippingInfo{ .X = 100.0f, .Width = 50.0f },
+		UE::Slate::FClippingInfo{ .X = 150.0f, .Width = 50.0f },
 	};
-	double WrapButtonX = 0.0;
+	TOptional<float> WrapButtonX;
 
 	SECTION("When size is constrained")
 	{
-		constexpr double AllottedWidth = 190;
+		constexpr float AllottedWidth = 190.0f;
 
 		TArray<UE::Slate::FClippingInfo> InOutClippingInfos = OriginalClippingInfos;
 		UE::Slate::PrioritizedResize(
@@ -287,37 +287,37 @@ TEST_CASE("Runtime::Slate::PrioritizedResizing positions a wrap button at index
 
 		SECTION("The wrap button is positioned snugly with the non-clipped widgets")
 		{
-			CHECK(InOutClippingInfos[0].X == 0);
-			CHECK(InOutClippingInfos[0].Width == 50);
+			CHECK(InOutClippingInfos[0].X == 0.0f);
+			CHECK(InOutClippingInfos[0].Width == 50.0f);
 
 			CHECK(WrapButtonX == 50);
 
-			CHECK(InOutClippingInfos[1].X == 90);
-			CHECK(InOutClippingInfos[1].Width == 50);
+			CHECK(InOutClippingInfos[1].X == 90.0f);
+			CHECK(InOutClippingInfos[1].Width == 50.0f);
 
-			CHECK(InOutClippingInfos[2].X == 140);
-			CHECK(InOutClippingInfos[2].Width == 50);
+			CHECK(InOutClippingInfos[2].X == 140.0f);
+			CHECK(InOutClippingInfos[2].Width == 50.0f);
 		}
 	}
 }
 
 TEST_CASE("Runtime::Slate::PrioritizedResizing adds wrap button when needed", "[Slate]")
 {
-	constexpr double WrapButtonWidth = 40;
-	const FMargin WrapButtonPadding = FMargin(0);
+	constexpr float WrapButtonWidth = 40.0f;
+	const FMargin WrapButtonPadding = FMargin(0.0f);
 	constexpr int32 WrapButtonIndex = 1;
 
 	const TArray<UE::Slate::FClippingInfo> OriginalClippingInfos = {
-		UE::Slate::FClippingInfo{ .X = 0, .Width = 50 },
-		UE::Slate::FClippingInfo{ .ResizeParams = { .VisibleInOverflow = false }, .X = 50, .Width = 50 },
-		UE::Slate::FClippingInfo{ .ResizeParams = { .VisibleInOverflow = false }, .X = 100, .Width = 50 },
-		UE::Slate::FClippingInfo{ .X = 150, .Width = 50 },
+		UE::Slate::FClippingInfo{ .X = 0.0f, .Width = 50.0f },
+		UE::Slate::FClippingInfo{ .ResizeParams = { .VisibleInOverflow = false }, .X = 50.0f, .Width = 50.0f },
+		UE::Slate::FClippingInfo{ .ResizeParams = { .VisibleInOverflow = false }, .X = 100.0f, .Width = 50.0f },
+		UE::Slate::FClippingInfo{ .X = 150.0f, .Width = 50.0f },
 	};
-	double WrapButtonX = 0.0;
+	TOptional<float> WrapButtonX;
 
 	SECTION("When size is constrained")
 	{
-		constexpr double AllottedWidth = 190;
+		constexpr float AllottedWidth = 190.0f;
 
 		TArray<UE::Slate::FClippingInfo> InOutClippingInfos = OriginalClippingInfos;
 		UE::Slate::PrioritizedResize(
@@ -338,36 +338,36 @@ TEST_CASE("Runtime::Slate::PrioritizedResizing adds wrap button when needed", "[
 
 		SECTION("The wrap button is positioned")
 		{
-			CHECK(WrapButtonX == 50);
+			CHECK(WrapButtonX == 50.0f);
 		}
 	}
 }
 
 TEST_CASE("Runtime::Slate::PrioritizedResizing with complex set of entries", "[Slate]")
 {
-	constexpr double WrapButtonWidth = 40;
-	const FMargin WrapButtonPadding = FMargin(0);
+	constexpr float WrapButtonWidth = 40.0f;
+	const FMargin WrapButtonPadding = FMargin(0.0f);
 	constexpr int32 WrapButtonIndex = 1;
 
 	// clang-format off
 	const TArray<UE::Slate::FClippingInfo> OriginalClippingInfos = {
-		UE::Slate::FClippingInfo{ .ResizeParams = { .ClippingPriority = 1000 }, .X = 0, .Width = 50 },
-		UE::Slate::FClippingInfo{ .ResizeParams = { .ClippingPriority = 500, .VisibleInOverflow = false }, .X = 50, .Width = 50 },
-		UE::Slate::FClippingInfo{ .ResizeParams = { .ClippingPriority = 500, .VisibleInOverflow = false }, .X = 100, .Width = 50 },
-		UE::Slate::FClippingInfo{ .ResizeParams = { .ClippingPriority = 800 }, .X = 150, .Width = 50 },
-		UE::Slate::FClippingInfo{ .ResizeParams = { .ClippingPriority = 400, .VisibleInOverflow = false }, .X = 200, .Width = 50 },
-		UE::Slate::FClippingInfo{ .ResizeParams = { .ClippingPriority = 400, .VisibleInOverflow = false }, .X = 250, .Width = 50 },
-		UE::Slate::FClippingInfo{ .ResizeParams = { .AllowClipping = false }, .X = 300, .Width = 50, .bIsStretchable = true },
-		UE::Slate::FClippingInfo{ .ResizeParams = { .ClippingPriority = 1000, .VisibleInOverflow = false }, .X = 350, .Width = 50 },
-		UE::Slate::FClippingInfo{ .ResizeParams = { .AllowClipping = false }, .X = 400, .Width = 50 },
-		UE::Slate::FClippingInfo{ .ResizeParams = { .AllowClipping = false, .VisibleInOverflow = false  }, .X = 450, .Width = 50 },
+		UE::Slate::FClippingInfo{ .ResizeParams = { .ClippingPriority = 1000 }, .X = 0.0f, .Width = 50.0f },
+		UE::Slate::FClippingInfo{ .ResizeParams = { .ClippingPriority = 500, .VisibleInOverflow = false }, .X = 50.0f, .Width = 50.0f },
+		UE::Slate::FClippingInfo{ .ResizeParams = { .ClippingPriority = 500, .VisibleInOverflow = false }, .X = 100, .Width = 50.0f },
+		UE::Slate::FClippingInfo{ .ResizeParams = { .ClippingPriority = 800 }, .X = 150.0f, .Width = 50.0f },
+		UE::Slate::FClippingInfo{ .ResizeParams = { .ClippingPriority = 400, .VisibleInOverflow = false }, .X = 200.0f, .Width = 50.0f },
+		UE::Slate::FClippingInfo{ .ResizeParams = { .ClippingPriority = 400, .VisibleInOverflow = false }, .X = 250.0f, .Width = 50.0f },
+		UE::Slate::FClippingInfo{ .ResizeParams = { .AllowClipping = false }, .X = 300.0f, .Width = 50.0f, .bIsStretchable = true },
+		UE::Slate::FClippingInfo{ .ResizeParams = { .ClippingPriority = 1000, .VisibleInOverflow = false }, .X = 350.0f, .Width = 50.0f },
+		UE::Slate::FClippingInfo{ .ResizeParams = { .AllowClipping = false }, .X = 400.0f, .Width = 50.0f },
+		UE::Slate::FClippingInfo{ .ResizeParams = { .AllowClipping = false, .VisibleInOverflow = false  }, .X = 450.0f, .Width = 50.0f },
 	};
 	// clang-format on
-	double WrapButtonX = -1.0;
+	TOptional<float> WrapButtonX;
 
 	GIVEN("Size is not constrained")
 	{
-		constexpr double AllottedWidth = 501;
+		constexpr float AllottedWidth = 501.0f;
 
 		TArray<UE::Slate::FClippingInfo> InOutClippingInfos = OriginalClippingInfos;
 		UE::Slate::PrioritizedResize(
@@ -384,13 +384,13 @@ TEST_CASE("Runtime::Slate::PrioritizedResizing with complex set of entries", "[S
 
 		THEN("The wrap button is not placed")
 		{
-			CHECK(WrapButtonX < 0.0);
+			CHECK(!WrapButtonX.IsSet());
 		}
 	}
 
 	GIVEN("Size is constrained by one pixel")
 	{
-		constexpr double AllottedWidth = 499;
+		constexpr float AllottedWidth = 499.0f;
 
 		TArray<UE::Slate::FClippingInfo> InOutClippingInfos = OriginalClippingInfos;
 		UE::Slate::PrioritizedResize(
@@ -417,13 +417,13 @@ TEST_CASE("Runtime::Slate::PrioritizedResizing with complex set of entries", "[S
 
 		THEN("The wrap button is not placed")
 		{
-			CHECK(WrapButtonX < 0.0);
+			CHECK(!WrapButtonX.IsSet());
 		}
 	}
 
 	GIVEN("Size is constrained by a little more than one button width")
 	{
-		constexpr double AllottedWidth = 445;
+		constexpr float AllottedWidth = 445.0f;
 
 		TArray<UE::Slate::FClippingInfo> InOutClippingInfos = OriginalClippingInfos;
 		UE::Slate::PrioritizedResize(
@@ -445,13 +445,13 @@ TEST_CASE("Runtime::Slate::PrioritizedResizing with complex set of entries", "[S
 
 		THEN("The wrap button is not placed")
 		{
-			CHECK(WrapButtonX < 0.0);
+			CHECK(!WrapButtonX.IsSet());
 		}
 	}
 
 	GIVEN("Size is constrained by a little more than two button widths")
 	{
-		constexpr double AllottedWidth = 395;
+		constexpr float AllottedWidth = 395.0f;
 
 		TArray<UE::Slate::FClippingInfo> InOutClippingInfos = OriginalClippingInfos;
 		UE::Slate::PrioritizedResize(
@@ -473,13 +473,13 @@ TEST_CASE("Runtime::Slate::PrioritizedResizing with complex set of entries", "[S
 
 		THEN("The wrap button is not placed")
 		{
-			CHECK(WrapButtonX < 0.0);
+			CHECK(!WrapButtonX.IsSet());
 		}
 	}
 
 	GIVEN("Size is constrained by a little more three button widths")
 	{
-		constexpr double AllottedWidth = 345;
+		constexpr float AllottedWidth = 345.0f;
 
 		TArray<UE::Slate::FClippingInfo> InOutClippingInfos = OriginalClippingInfos;
 		UE::Slate::PrioritizedResize(
@@ -501,28 +501,29 @@ TEST_CASE("Runtime::Slate::PrioritizedResizing with complex set of entries", "[S
 
 		THEN("The wrap button is not placed")
 		{
-			CHECK(WrapButtonX < 0.0);
+			CHECK(!WrapButtonX.IsSet());
 		}
 	}
 }
 
 TEST_CASE("Runtime::Slate::PrioritizedResizing takes the wrapping button width into account", "[Slate]")
 {
-	constexpr double WrapButtonWidth = 40.0;
-	const FMargin WrapButtonPadding = FMargin(3.0f, 0);
+	constexpr float WrapButtonWidth = 40.0f;
+	const FMargin WrapButtonPadding = FMargin(3.0f, 0.0f);
 	constexpr int32 WrapButtonIndex = -2;
 
 	const TArray<UE::Slate::FClippingInfo> OriginalClippingInfos = {
-		UE::Slate::FClippingInfo{ .X = 0,   .Width = 70 },
-		UE::Slate::FClippingInfo{ .X = 70,  .Width = 70 },
-		UE::Slate::FClippingInfo{ .X = 140, .Width = 70 },
-		UE::Slate::FClippingInfo{ .X = 210, .Width = 70 },
+		UE::Slate::FClippingInfo{ .X = 0.0f,   .Width = 70.0f },
+		UE::Slate::FClippingInfo{ .X = 70.0f,  .Width = 70.0f },
+		UE::Slate::FClippingInfo{ .X = 140.0f, .Width = 70.0f },
+		UE::Slate::FClippingInfo{ .X = 210.0f, .Width = 70.0f },
 	};
-	double WrapButtonX = -1.0;
+	
+	TOptional<float> WrapButtonX;
 
 	GIVEN("Size is constrained by a few pixels")
 	{
-		constexpr double AllottedWidth = 274.0;
+		constexpr float AllottedWidth = 274.0f;
 
 		TArray<UE::Slate::FClippingInfo> InOutClippingInfos = OriginalClippingInfos;
 		UE::Slate::PrioritizedResize(
@@ -532,7 +533,7 @@ TEST_CASE("Runtime::Slate::PrioritizedResizing takes the wrapping button width i
 		THEN("The wrapping button is placed")
 		{
 			// Position + Left padding
-			CHECK(WrapButtonX == 140 + 3);
+			CHECK(WrapButtonX.Get(0.0f) == 140.0f + 3.0f);
 		}
 
 		THEN("The last widget is clipped")
@@ -553,10 +554,12 @@ TEST_CASE("Runtime::Slate::PrioritizedResizing takes the wrapping button width i
 			CHECK(NumClipped == 1);
 		}
 	}
+	
+	WrapButtonX.Reset();
 
 	GIVEN("Size is constrained enough so the wrap button will force another widget to clip")
 	{
-		constexpr double AllottedWidth = 230.0;
+		constexpr float AllottedWidth = 230.0f;
 
 		TArray<UE::Slate::FClippingInfo> InOutClippingInfos = OriginalClippingInfos;
 		UE::Slate::PrioritizedResize(
@@ -566,7 +569,7 @@ TEST_CASE("Runtime::Slate::PrioritizedResizing takes the wrapping button width i
 		THEN("The wrapping button is placed")
 		{
 			// Position + Left padding
-			CHECK(WrapButtonX == 70 + 3);
+			CHECK(WrapButtonX.Get(0.0f) == 70.0f + 3.0f);
 		}
 
 		THEN("The second last widget is clipped")
diff --git a/Engine/Source/Runtime/Slate/Private/Framework/MultiBox/SClippingHorizontalBox.cpp b/Engine/Source/Runtime/Slate/Private/Framework/MultiBox/SClippingHorizontalBox.cpp
--- a/Engine/Source/Runtime/Slate/Private/Framework/MultiBox/SClippingHorizontalBox.cpp
+++ b/Engine/Source/Runtime/Slate/Private/Framework/MultiBox/SClippingHorizontalBox.cpp
@@ -13,12 +13,12 @@ namespace UE::Slate
 {
 
 void PrioritizedResize(
-	double InAllottedWidth,
-	double InWrapButtonWidth,
+	float InAllottedWidth,
+	float InWrapButtonWidth,
 	const FMargin& InWrapButtonPadding,
 	int32 InWrapButtonIndex,
 	TArray<FClippingInfo>& InOutClippingInfos,
-	double& OutWrapButtonX
+	TOptional<float>& OutWrapButtonX
 )
 {
 	double WidthOfAllChildren = 0.0;
@@ -285,9 +285,9 @@ void SClippingHorizontalBox::OnArrangeChildren(const FGeometry& AllottedGeometry
 		}
 	}
 
-	const double InAllottedWidth = AllottedGeometry.GetLocalSize().X;
+	const float InAllottedWidth = AllottedGeometry.GetLocalSize().X;
 	const FToolBarStyle& ToolBarStyle = StyleSet->GetWidgetStyle<FToolBarStyle>(StyleName);
-	double WrapButtonX = 0.0;
+	TOptional<float> WrapButtonX;
 
 	if (bAllowWrapButton)
 	{
@@ -329,10 +329,11 @@ void SClippingHorizontalBox::OnArrangeChildren(const FGeometry& AllottedGeometry
 	}
 
 	// Position the wrap button.
+	if (const float* WrapX = WrapButtonX.GetPtrOrNull())
 	{
 		FGeometry& WrapButtonGeometry = ArrangedChildren[ArrangedChildren.Num() - 1].Geometry;
 		const FVector2D WrapButtonSize = FVector2D(WrapButtonWidth, WrapButtonGeometry.GetLocalSize().Y);
-		WrapButtonGeometry = AllottedGeometry.MakeChild(WrapButtonSize, FSlateLayoutTransform(FVector2f(WrapButtonX, 0)));
+		WrapButtonGeometry = AllottedGeometry.MakeChild(WrapButtonSize, FSlateLayoutTransform(FVector2f(*WrapX, 0)));
 	}
 
 	// Finally, position all children using the clipping information.
diff --git a/Engine/Source/Runtime/Slate/Private/Framework/MultiBox/SClippingHorizontalBox.h b/Engine/Source/Runtime/Slate/Private/Framework/MultiBox/SClippingHorizontalBox.h
--- a/Engine/Source/Runtime/Slate/Private/Framework/MultiBox/SClippingHorizontalBox.h
+++ b/Engine/Source/Runtime/Slate/Private/Framework/MultiBox/SClippingHorizontalBox.h
@@ -31,12 +31,12 @@ struct FClippingInfo
 };
 
 void PrioritizedResize(
-	double InAllottedWidth,
-	double InWrapButtonWidth,
+	float InAllottedWidth,
+	float InWrapButtonWidth,
 	const FMargin& InWrapButtonPadding,
 	int32 InWrapButtonIndex,
 	TArray<FClippingInfo>& InOutClippingInfos,
-	double& OutWrapButtonX
+	TOptional<float>& OutWrapButtonX
 );
 
 }; // namespace UE::Slate
